LearningCore: Added EntryIntersectsRange function

#Synced-CL 50249268

diff --git a/Engine/Plugins/Experimental/LearningCore/Source/Learning/Private/LearningFrameRangeSet.cpp b/Engine/Plugins/Experimental/LearningCore/Source/Learning/Private/LearningFrameRangeSet.cpp
--- a/Engine/Plugins/Experimental/LearningCore/Source/Learning/Private/LearningFrameRangeSet.cpp
+++ b/Engine/Plugins/Experimental/LearningCore/Source/Learning/Private/LearningFrameRangeSet.cpp
@@ -1514,6 +1514,15 @@ namespace UE::Learning
 			Start, Length);
 	}
 
+	bool FFrameRangeSet::EntryIntersectsRange(const int32 EntryIdx, const int32 Start, const int32 Length) const
+	{
+		return FrameRangeSet::Private::RangesIntersectsRange(
+			GetEntryRangeStarts(EntryIdx),
+			GetEntryRangeLengths(EntryIdx),
+			Start,
+			Length);
+	}
+
 	bool FFrameRangeSet::ContainsTime(const int32 Sequence, const float Time, const float FrameDeltaTime) const
 	{
 		const int32 EntryIdx = FindSequenceEntry(Sequence);
diff --git a/Engine/Plugins/Experimental/LearningCore/Source/Learning/Public/LearningFrameRangeSet.h b/Engine/Plugins/Experimental/LearningCore/Source/Learning/Public/LearningFrameRangeSet.h
--- a/Engine/Plugins/Experimental/LearningCore/Source/Learning/Public/LearningFrameRangeSet.h
+++ b/Engine/Plugins/Experimental/LearningCore/Source/Learning/Public/LearningFrameRangeSet.h
@@ -123,6 +123,9 @@ namespace UE::Learning
 		/** Checks if the given Entry of this FrameRangeSet fully contains a range */
 		UE_API bool EntryContainsRange(const int32 EntryIdx, const int32 Start, const int32 Length) const;
 
+		/** Checks if the given Entry of this FrameRangeSet intersects any part of a range */
+		UE_API bool EntryIntersectsRange(const int32 EntryIdx, const int32 Start, const int32 Length) const;
+
 		/** Checks if this FrameRangeSet contains a given sequence and time in that sequence */
 		UE_API bool ContainsTime(const int32 Sequence, const float Time, const float FrameDeltaTime) const;
 
