LearningCore: Added Select frame attribute function

#Synced-CL 50320943

diff --git a/Engine/Plugins/Experimental/LearningCore/Source/Learning/Private/Learning.ispc b/Engine/Plugins/Experimental/LearningCore/Source/Learning/Private/Learning.ispc
--- a/Engine/Plugins/Experimental/LearningCore/Source/Learning/Private/Learning.ispc
+++ b/Engine/Plugins/Experimental/LearningCore/Source/Learning/Private/Learning.ispc
@@ -853,6 +853,19 @@ export void LearningFilterMajorityVote(
 	}
 }
 
+export void LearningFrameAttributeSelect(
+	uniform float Out[], 
+	uniform const float Cond[],
+	uniform const float True[],
+	uniform const float False[],
+	uniform int ValueNum)
+{
+	foreach(ValueIdx = 0 ... ValueNum)
+	{
+		Out[ValueIdx] = select(Cond[ValueIdx], True[ValueIdx], False[ValueIdx]);
+	}
+}
+
 export void LearningFrameAttributeAdd(
 	uniform float Out[], 
 	uniform const float Lhs[],
diff --git a/Engine/Plugins/Experimental/LearningCore/Source/Learning/Private/LearningFrameAttribute.cpp b/Engine/Plugins/Experimental/LearningCore/Source/Learning/Private/LearningFrameAttribute.cpp
--- a/Engine/Plugins/Experimental/LearningCore/Source/Learning/Private/LearningFrameAttribute.cpp
+++ b/Engine/Plugins/Experimental/LearningCore/Source/Learning/Private/LearningFrameAttribute.cpp
@@ -98,6 +98,28 @@ namespace UE::Learning
 #endif
 		}
 
+		static inline void Select(TLearningArrayView<1, float> Out, const TLearningArrayView<1, const float> Cond, const TLearningArrayView<1, const float> True, const TLearningArrayView<1, const float> False)
+		{
+			check(Out.Num() == Cond.Num());
+			check(Out.Num() == True.Num());
+			check(Out.Num() == False.Num());
+
+#if UE_LEARNING_ISPC
+			ispc::LearningFrameAttributeSelect(
+				Out.GetData(),
+				Cond.GetData(),
+				True.GetData(),
+				False.GetData(),
+				Out.Num());
+#else
+			const int32 ValueNum = Out.Num();
+			for (int32 ValueIdx = 0; ValueIdx < ValueNum; ValueIdx++)
+			{
+				Out[ValueIdx] = Cond[ValueIdx] ? True[ValueIdx] : False[ValueIdx];
+			}
+#endif
+		}
+
 		static inline void Add(TLearningArrayView<1, float> Out, const TLearningArrayView<1, const float> Lhs, const TLearningArrayView<1, const float> Rhs)
 		{
 			check(Out.Num() == Lhs.Num());
@@ -2193,6 +2215,39 @@ namespace UE::Learning
 			return Private::FindMaximum(OutChannelIdx, OutFrameIdx, OutMaximumValue, In.GetAttributeData());
 		}
 
+		void Select(FFrameAttribute& Out, const FFrameAttribute& Cond, const FFrameAttribute& True, const FFrameAttribute& False)
+		{
+			check(Cond.GetChannelNum() == 1);
+			check(True.GetChannelNum() == False.GetChannelNum());
+
+			const int32 ChannelNum = True.GetChannelNum();
+
+			NaryOp(Out, ChannelNum, {
+				(ConstFrameAttributePtr)&Cond,
+				(ConstFrameAttributePtr)&True,
+				(ConstFrameAttributePtr)&False }, [ChannelNum](
+					FFrameAttribute& Out,
+					const TArrayView<const ConstFrameAttributePtr> In,
+					const TLearningArrayView<1, const int32> OutRangeOffsets,
+					const TArrayView<const TLearningArrayView<1, const int32>> InRangeOffsets,
+					const TLearningArrayView<1, const int32> RangeLengths) {
+
+						for (int32 ChannelIdx = 0; ChannelIdx < ChannelNum; ChannelIdx++)
+						{
+							const int32 RangeNum = RangeLengths.Num();
+
+							for (int32 RangeIdx = 0; RangeIdx < RangeNum; RangeIdx++)
+							{
+								Private::Select(
+									Out.GetChannelRangeAttributeData(ChannelIdx, OutRangeOffsets[RangeIdx], RangeLengths[RangeIdx]),
+									In[0]->GetChannelRangeAttributeData(0, InRangeOffsets[0][RangeIdx], RangeLengths[RangeIdx]),
+									In[1]->GetChannelRangeAttributeData(ChannelIdx, InRangeOffsets[1][RangeIdx], RangeLengths[RangeIdx]),
+									In[2]->GetChannelRangeAttributeData(ChannelIdx, InRangeOffsets[2][RangeIdx], RangeLengths[RangeIdx]));
+							}
+						}
+				});
+		}
+
 		void Zeros(FFrameAttribute& Out, const FFrameRangeSet& FrameRangeSet, const int32 ChannelNum)
 		{
 			Out.FrameRangeSet = FrameRangeSet;
diff --git a/Engine/Plugins/Experimental/LearningCore/Source/Learning/Public/LearningFrameAttribute.h b/Engine/Plugins/Experimental/LearningCore/Source/Learning/Public/LearningFrameAttribute.h
--- a/Engine/Plugins/Experimental/LearningCore/Source/Learning/Public/LearningFrameAttribute.h
+++ b/Engine/Plugins/Experimental/LearningCore/Source/Learning/Public/LearningFrameAttribute.h
@@ -202,6 +202,10 @@ namespace UE::Learning
 			const FFrameAttribute& In);
 
 
+		/** Select between True and False frame attributes based on Cond. */
+		UE_API void Select(FFrameAttribute& Out, const FFrameAttribute& Cond, const FFrameAttribute& True, const FFrameAttribute& False);
+
+
 		/** Create a frame attribute made up of zeros */
 		UE_API void Zeros(FFrameAttribute& Out, const FFrameRangeSet& FrameRangeSet, const int32 ChannelNum);
 
