LearningCore: Added Inertialization Distance Function and various other ISPC optimizations

#Synced-CL 50413561

diff --git a/Engine/Plugins/Experimental/LearningCore/Source/Learning/Private/Learning.ispc b/Engine/Plugins/Experimental/LearningCore/Source/Learning/Private/Learning.ispc
--- a/Engine/Plugins/Experimental/LearningCore/Source/Learning/Private/Learning.ispc
+++ b/Engine/Plugins/Experimental/LearningCore/Source/Learning/Private/Learning.ispc
@@ -42,6 +42,116 @@ inline float LearningApproxAcos(float X)
 	return X >= 0.0 ? Q : UE_PI - Q;
 }
 
+inline float LearningInertializationPartialIntegral(
+	const float X,
+	const float V,
+	const float Time,
+	uniform const float BlendTime)
+{
+	uniform const float Z = max(BlendTime, UE_SMALL_NUMBER);
+	const float T = clamp(Time / Z, 0.0f, 1.0f);
+	return
+		X * ((Z / 2) * T * T * T * T - Z * T * T * T + Z * T) +
+		V * (((Z * Z) / 4) * T * T * T * T - ((2 * Z * Z) / 3) * T * T * T + ((Z * Z) / 2) * T * T);
+}
+
+inline float LearningInertializationFullIntegral(
+	const float X,
+	const float V,
+	uniform const float BlendTime)
+{
+	uniform const float Z = max(BlendTime, UE_SMALL_NUMBER);
+	return X * (Z / 2) + V * ((Z * Z) / 12);
+}
+
+inline float LearningInertializationPosition(
+	const float X,
+	const float V,
+	const float Time,
+	uniform const float BlendTime)
+{
+	uniform const float Z = max(BlendTime, UE_SMALL_NUMBER);
+	const float T = clamp(Time / Z, 0.0f, 1.0f);
+	return X * (2 * T * T * T - 3 * T * T + 1) + V * (Z * T * T * T - Z * 2 * T * T + Z * T);
+}
+
+inline float LearningInertializationVelocity(
+	const float X,
+	const float V,
+	const float Time,
+	uniform const float BlendTime)
+{
+	uniform const float Z = max(BlendTime, UE_SMALL_NUMBER);
+	const float T = clamp(Time / Z, 0.0f, 1.0f);
+	return X * ((6 / Z) * T * T - (6 / Z) * T) + V * (3 * T * T - 4 * T + 1);
+}
+
+inline float LearningInertializationPositionIntersectionTime(
+	const float X,
+	const float V,
+	uniform const float BlendTime)
+{
+	uniform const float Z = max(BlendTime, UE_SMALL_NUMBER);
+	return (-Z * X) / (2 * X + Z * V);
+}
+
+inline float LearningInertializationVelocityIntersectionTime(
+	const float X,
+	const float V,
+	uniform const float BlendTime)
+{
+	uniform const float Z = max(BlendTime, UE_SMALL_NUMBER);
+	return (Z * Z * V) / (6 * X + 3 * Z * V);
+}
+
+inline float LearningInertializationAccelerationIntersectionTime(
+	const float X,
+	const float V,
+	uniform const float BlendTime)
+{
+	uniform const float Z = max(BlendTime, UE_SMALL_NUMBER);
+	return (3 * Z * X + 2 * Z * Z * V) / (6 * X + 3 * Z * V);
+}
+
+inline float LearningInertializationMaximumPositionMagnitude(
+	const float X,
+	const float V,
+	uniform const float BlendTime)
+{
+	uniform const float Z = max(BlendTime, UE_SMALL_NUMBER);
+	const float T = LearningInertializationVelocityIntersectionTime(X, V, Z);
+
+	return T > 0.0f && T < Z ?
+		max(abs(X), abs(LearningInertializationPosition(X, V, T, Z))) : abs(X);
+}
+
+inline float LearningInertializationMaximumVelocityMagnitude(
+	const float X,
+	const float V,
+	uniform const float BlendTime)
+{
+	uniform const float Z = max(BlendTime, UE_SMALL_NUMBER);
+	const float T = LearningInertializationAccelerationIntersectionTime(X, V, Z);
+
+	return T > 0.0f && T < Z ?
+		max(abs(V), abs(LearningInertializationVelocity(X, V, T, Z))) : abs(V);
+}
+
+inline float LearningInertializationIntegralMagnitude(
+	const float X,
+	const float V,
+	uniform const float BlendTime)
+{
+	uniform const float Z = max(BlendTime, UE_SMALL_NUMBER);
+	const float T = LearningInertializationPositionIntersectionTime(X, V, Z);
+	const float IntegralT = LearningInertializationPartialIntegral(X, V, T, Z);
+	const float IntegralZ = LearningInertializationFullIntegral(X, V, Z);
+
+	return T > 0.0f && T < Z ?
+		abs(IntegralZ - IntegralT) + abs(IntegralT) :
+		abs(IntegralZ);
+}
+
 /////////////////////////////////////////////////////////
 // 
 // Functions matching those in LearningRandom.cpp
@@ -705,6 +815,58 @@ export void LearningKMeansComputeClusteredIndex(
 	}
 }
 
+export void LearningComputeRangeOffsets(
+	uniform int32 Offsets[],
+	uniform const int32 Lengths[],
+	uniform const int RangeNum,
+	uniform const int32 InitialOffset)
+{
+	uniform int32 Offset = InitialOffset;
+	for (uniform int RangeIdx = 0; RangeIdx < RangeNum; RangeIdx++)
+	{
+		Offsets[RangeIdx] = Offset;
+		Offset += Lengths[RangeIdx];
+	}
+}
+
+export void LearningFramesCopy(
+	uniform int32 DstFrames[],
+	uniform const int32 SrcFrames[],
+	uniform const int FrameNum)
+{
+	foreach (FrameIdx = 0 ... FrameNum)
+	{
+		DstFrames[FrameIdx] = SrcFrames[FrameIdx];
+	}
+}
+
+export void LearningRangesCopy(
+	uniform int32 DstStarts[],
+	uniform int32 DstLengths[],
+	uniform const int32 SrcStarts[],
+	uniform const int32 SrcLengths[],
+	uniform const int RangeNum)
+{
+	foreach (RangeIdx = 0 ... RangeNum)
+	{
+		DstStarts[RangeIdx] = SrcStarts[RangeIdx];
+		DstLengths[RangeIdx] = SrcLengths[RangeIdx];
+	}
+}
+
+export void LearningRangesCopyFrames(
+	uniform int32 DstStarts[],
+	uniform int32 DstLengths[],
+	uniform const int32 SrcFrames[],
+	uniform const int RangeNum)
+{
+	foreach(RangeIdx = 0 ... RangeNum)
+	{
+		DstStarts[RangeIdx] = SrcFrames[RangeIdx];
+		DstLengths[RangeIdx] = 1;
+	}
+}
+
 export uniform bool LearningRangesEqual(
 	uniform const int32 LhsStarts[],
 	uniform const int32 LhsLengths[],
@@ -765,6 +927,26 @@ export uniform bool LearningRangesFind(
 	return false;
 }
 
+export uniform bool LearningRangesContainsTime(
+	uniform const int32 Starts[],
+	uniform const int32 Lengths[],
+	uniform int32 RangeNum,
+	uniform const float Time,
+	uniform const float FrameDeltaTime)
+{
+	for (uniform int RangeIdx = 0; RangeIdx < RangeNum; RangeIdx++)
+	{
+		uniform const float StartTime = Starts[RangeIdx] * FrameDeltaTime;
+		uniform const float EndTime = (Starts[RangeIdx] + Lengths[RangeIdx] - 1) * FrameDeltaTime;
+
+		if (Time >= StartTime && Time < EndTime)
+		{
+			return true;
+		}
+	}
+	return false;
+}
+
 export uniform bool LearningRangesIntersectsRange(
 	uniform const int32 Starts[],
 	uniform const int32 Lengths[],
@@ -799,6 +981,1077 @@ export uniform bool LearningRangesContainsRange(
 	return false;
 }
 
+export uniform bool LearningRangesFindTime(
+	uniform int32& OutRangeIdx,
+	uniform float& OutRangeTime,
+	uniform const int32 Starts[],
+	uniform const int32 Lengths[],
+	uniform const int RangeNum,
+	uniform const float Time,
+	uniform const float FrameDeltaTime)
+{
+	for (uniform int RangeIdx = 0; RangeIdx < RangeNum; RangeIdx++)
+	{
+		uniform const float StartTime = Starts[RangeIdx] * FrameDeltaTime;
+		uniform const float EndTime = (Starts[RangeIdx] + Lengths[RangeIdx] - 1) * FrameDeltaTime;
+
+		if (Time >= StartTime && Time < EndTime)
+		{
+			OutRangeIdx = RangeIdx;
+			OutRangeTime = Time - StartTime;
+			return true;
+		}
+	}
+
+	OutRangeIdx = INDEX_NONE;
+	OutRangeTime = -1.0f;
+	return false;
+}
+
+
+export uniform int32 LearningFramesRangesUnion(
+	uniform int32 OutStarts[],
+	uniform int32 OutLengths[],
+	uniform const uniform int32 LhsFrames[],
+	uniform const uniform int32 RhsStarts[],
+	uniform const uniform int32 RhsLengths[],
+	uniform const int LhsNum,
+	uniform const int RhsNum)
+{
+	if (LhsNum == 0)
+	{
+		LearningRangesCopy(
+			OutStarts,
+			OutLengths,
+			RhsStarts,
+			RhsLengths,
+			RhsNum);
+
+		return RhsNum;
+	}
+
+	if (RhsNum == 0)
+	{
+		LearningRangesCopyFrames(
+			OutStarts,
+			OutLengths,
+			LhsFrames,
+			LhsNum);
+
+		return LhsNum;
+	}
+
+	// Activation state of each list of ranges
+	uniform bool bOutActive = false;
+	uniform bool bLhsActive = false;
+	uniform bool bRhsActive = false;
+
+	// Event index for each list of ranges
+	uniform int32 LhsIndex = 0;
+	uniform int32 RhsIndex = 0;
+	uniform int32 OutIndex = 0;
+
+	// While both ranges have events to process
+	while (LhsIndex < LhsNum * 2 && RhsIndex < RhsNum * 2)
+	{
+		// Are the next lhs, and rhs events active or inactive
+		uniform const bool bLhsActiveNext = LhsIndex % 2 == 0;
+		uniform const bool bRhsActiveNext = RhsIndex % 2 == 0;
+
+		// Time of the next lhs, and rhs events
+		uniform const int32 LhsT = bLhsActiveNext ? LhsFrames[LhsIndex / 2] : LhsFrames[LhsIndex / 2] + 1;
+		uniform const int32 RhsT = bRhsActiveNext ? RhsStarts[RhsIndex / 2] : RhsStarts[RhsIndex / 2] + RhsLengths[RhsIndex / 2];
+
+		// Event from lhs is coming first
+		if (LhsT < RhsT)
+		{
+			// Activate output
+			if (!bOutActive && bLhsActiveNext)
+			{
+				bOutActive = true;
+				OutStarts[OutIndex] = LhsT;
+			}
+			// Deactivate output
+			else if (bOutActive && !bLhsActiveNext && !bRhsActive)
+			{
+
+				bOutActive = false;
+				OutLengths[OutIndex] = LhsT - OutStarts[OutIndex];
+				OutIndex++;
+			}
+
+			bLhsActive = bLhsActiveNext;
+			LhsIndex++;
+		}
+		// Event from rhs is coming first
+		else if (RhsT < LhsT)
+		{
+			// Activate output
+			if (!bOutActive && bRhsActiveNext)
+			{
+				bOutActive = true;
+				OutStarts[OutIndex] = RhsT;
+			}
+			// Deactivate output
+			else if (bOutActive && !bLhsActive && !bRhsActiveNext)
+			{
+				bOutActive = false;
+				OutLengths[OutIndex] = RhsT - OutStarts[OutIndex];
+				OutIndex++;
+			}
+
+			bRhsActive = bRhsActiveNext;
+			RhsIndex++;
+		}
+		// Event from lhs and rhs coming at same time
+		else
+		{
+			// Activate output
+			if (!bOutActive && (bLhsActiveNext || bRhsActiveNext))
+			{
+				bOutActive = true;
+				OutStarts[OutIndex] = LhsT;
+			}
+			// Deactivate output
+			else if (bOutActive && !(bLhsActiveNext || bRhsActiveNext))
+			{
+				bOutActive = false;
+				OutLengths[OutIndex] = LhsT - OutStarts[OutIndex];
+				OutIndex++;
+			}
+
+			bLhsActive = bLhsActiveNext;
+			bRhsActive = bRhsActiveNext;
+			LhsIndex++; RhsIndex++;
+		}
+	}
+
+	// Process any remaining lhs events
+	while (LhsIndex < LhsNum * 2)
+	{
+		uniform const bool bLhsActiveNext = LhsIndex % 2 == 0;
+		uniform const int32 LhsT = bLhsActiveNext ? LhsFrames[LhsIndex / 2] : LhsFrames[LhsIndex / 2] + 1;
+
+		// Activate output
+		if (!bOutActive && bLhsActiveNext)
+		{
+			bOutActive = true;
+			OutStarts[OutIndex] = LhsT;
+		}
+		// Deactivate output
+		else if (bOutActive && !bLhsActiveNext)
+		{
+			bOutActive = false;
+			OutLengths[OutIndex] = LhsT - OutStarts[OutIndex];
+			OutIndex++;
+		}
+
+		bLhsActive = bLhsActiveNext;
+		LhsIndex++;
+	}
+
+	// Process any remaining rhs events
+	while (RhsIndex < RhsNum * 2)
+	{
+		uniform const bool bRhsActiveNext = RhsIndex % 2 == 0;
+		uniform const int32 RhsT = bRhsActiveNext ? RhsStarts[RhsIndex / 2] : RhsStarts[RhsIndex / 2] + RhsLengths[RhsIndex / 2];
+
+		// Activate output
+		if (!bOutActive && bRhsActiveNext)
+		{
+			bOutActive = true;
+			OutStarts[OutIndex] = RhsT;
+		}
+		// Deactivate output
+		else if (bOutActive && !bRhsActiveNext)
+		{
+			bOutActive = false;
+			OutLengths[OutIndex] = RhsT - OutStarts[OutIndex];
+			OutIndex++;
+		}
+
+		bRhsActive = bRhsActiveNext;
+		RhsIndex++;
+	}
+
+	// Return number of ranges added
+	return OutIndex;
+}
+
+export uniform int32 LearningFramesRangesIntersection(
+	uniform int32 OutFrames[],
+	uniform const int32 LhsFrames[],
+	uniform const int32 RhsStarts[],
+	uniform const int32 RhsLengths[],
+	uniform const int LhsNum,
+	uniform const int RhsNum)
+{
+
+	if (LhsNum == 0) { return 0; }
+	if (RhsNum == 0) { return 0; }
+
+	// Activation state of ranges
+	uniform bool bRhsActive = false;
+
+	// Event index for each list
+	uniform int32 LhsIndex = 0;
+	uniform int32 RhsIndex = 0;
+	uniform int32 OutIndex = 0;
+
+	// While both ranges have events to process
+	while (LhsIndex < LhsNum && RhsIndex < RhsNum * 2)
+	{
+		// Are the rhs events active or inactive
+		uniform const bool bRhsActiveNext = RhsIndex % 2 == 0;
+
+		// Time of the next lhs, and rhs events
+		uniform const int32 LhsT = LhsFrames[LhsIndex];
+		uniform const int32 RhsT = bRhsActiveNext ? RhsStarts[RhsIndex / 2] : RhsStarts[RhsIndex / 2] + RhsLengths[RhsIndex / 2];
+
+		// Event from lhs coming first
+		if (LhsT < RhsT)
+		{
+			// Add to output if range is active
+			if (bRhsActive)
+			{
+				OutFrames[OutIndex] = LhsT;
+				OutIndex++;
+			}
+
+			LhsIndex++;
+		}
+		// Event from rhs coming first
+		else if (RhsT < LhsT)
+		{
+			bRhsActive = bRhsActiveNext;
+			RhsIndex++;
+		}
+		// Event from lhs and rhs coming at same time
+		else
+		{
+			bRhsActive = bRhsActiveNext;
+			RhsIndex++;
+
+			// Add to output if range is active
+			if (bRhsActive)
+			{
+				OutFrames[OutIndex] = LhsT;
+				OutIndex++;
+			}
+
+			LhsIndex++;
+		}
+	}
+
+	// Return number of frames added
+	return OutIndex;
+}
+
+export uniform int32 LearningFramesRangesDifference(
+	uniform int32 OutFrames[],
+	uniform const int32 LhsFrames[],
+	uniform const int32 RhsStarts[],
+	uniform const int32 RhsLengths[],
+	uniform const int LhsNum,
+	uniform const int RhsNum)
+{
+	if (LhsNum == 0) { return 0; }
+
+	if (RhsNum == 0)
+	{
+		LearningFramesCopy(OutFrames, LhsFrames, LhsNum);
+		return LhsNum;
+	}
+
+	// Activation state of ranges
+	uniform bool bRhsActive = false;
+
+	// Event index for each list
+	uniform int32 LhsIndex = 0;
+	uniform int32 RhsIndex = 0;
+	uniform int32 OutIndex = 0;
+
+	// While both ranges have events to process
+	while (LhsIndex < LhsNum && RhsIndex < RhsNum * 2)
+	{
+		// Are the rhs events active or inactive
+		uniform const bool bRhsActiveNext = RhsIndex % 2 == 0;
+
+		// Time of the next lhs, and rhs events
+		uniform const int32 LhsT = LhsFrames[LhsIndex];
+		uniform const int32 RhsT = bRhsActiveNext ? RhsStarts[RhsIndex / 2] : RhsStarts[RhsIndex / 2] + RhsLengths[RhsIndex / 2];
+
+		// Event from lhs coming first
+		if (LhsT < RhsT)
+		{
+			// Add to output if range is active
+			if (!bRhsActive)
+			{
+				OutFrames[OutIndex] = LhsT;
+				OutIndex++;
+			}
+
+			LhsIndex++;
+		}
+		// Event from rhs coming first
+		else if (RhsT < LhsT)
+		{
+			bRhsActive = bRhsActiveNext;
+			RhsIndex++;
+		}
+		// Event from lhs and rhs coming at same time
+		else
+		{
+			bRhsActive = bRhsActiveNext;
+			RhsIndex++;
+
+			// Add to output if range is active
+			if (!bRhsActive)
+			{
+				OutFrames[OutIndex] = LhsT;
+				OutIndex++;
+			}
+
+			LhsIndex++;
+		}
+	}
+
+	// Return number of events added
+	return OutIndex;
+}
+
+
+export uniform int32 LearningRangesFramesDifference(
+	uniform int32 OutStarts[],
+	uniform int32 OutLengths[],
+	uniform const int32 LhsStarts[],
+	uniform const int32 LhsLengths[],
+	uniform const int32 RhsFrames[],
+	uniform const int LhsNum,
+	uniform const int RhsNum)
+{
+
+	if (LhsNum == 0) { return 0; }
+
+	if (RhsNum == 0)
+	{
+		LearningRangesCopy(
+			OutStarts,
+			OutLengths,
+			LhsStarts,
+			LhsLengths,
+			LhsNum);
+
+		return LhsNum;
+	}
+
+	// Activation state of each list of ranges
+	uniform bool bOutActive = false;
+	uniform bool bLhsActive = false;
+	uniform bool bRhsActive = false;
+
+	// Event index for each list of ranges
+	uniform int32 LhsIndex = 0;
+	uniform int32 RhsIndex = 0;
+	uniform int32 OutIndex = 0;
+
+	// While both ranges have events to process
+	while (LhsIndex < LhsNum * 2 && RhsIndex < RhsNum * 2)
+	{
+		// Are the next lhs, and rhs events active or inactive
+		uniform const bool bLhsActiveNext = LhsIndex % 2 == 0;
+		uniform const bool bRhsActiveNext = RhsIndex % 2 == 0;
+
+		// Time of the next lhs, and rhs events
+		uniform const int32 LhsT = bLhsActiveNext ? LhsStarts[LhsIndex / 2] : LhsStarts[LhsIndex / 2] + LhsLengths[LhsIndex / 2];
+		uniform const int32 RhsT = bRhsActiveNext ? RhsFrames[RhsIndex / 2] : RhsFrames[RhsIndex / 2] + 1;
+
+		// Event coming from lhs first
+		if (LhsT < RhsT)
+		{
+			// Activate output
+			if (!bOutActive && !bRhsActive && bLhsActiveNext)
+			{
+				bOutActive = true;
+				OutStarts[OutIndex] = LhsT;
+			}
+			// Deactivate output
+			else if (bOutActive && !bLhsActiveNext)
+			{
+				bOutActive = false;
+				OutLengths[OutIndex] = LhsT - OutStarts[OutIndex];
+				OutIndex++;
+			}
+
+			bLhsActive = bLhsActiveNext;
+			LhsIndex++;
+		}
+		// Event coming from rhs first
+		else if (RhsT < LhsT)
+		{
+			// Activate output
+			if (!bOutActive && bLhsActive && !bRhsActiveNext)
+			{
+				bOutActive = true;
+				OutStarts[OutIndex] = RhsT;
+			}
+			// Deactivate output
+			else if (bOutActive && bRhsActiveNext)
+			{
+				bOutActive = false;
+				OutLengths[OutIndex] = RhsT - OutStarts[OutIndex];
+				OutIndex++;
+			}
+
+			bRhsActive = bRhsActiveNext;
+			RhsIndex++;
+		}
+		// Event from lhs and rhs coming at same time
+		else
+		{
+			// Activate output
+			if (!bOutActive && bLhsActiveNext && !bRhsActiveNext)
+			{
+				bOutActive = true;
+				OutStarts[OutIndex] = LhsT;
+			}
+			// Deactivate output
+			else if (bOutActive && bRhsActiveNext)
+			{
+				bOutActive = false;
+				OutLengths[OutIndex] = LhsT - OutStarts[OutIndex];
+				OutIndex++;
+			}
+
+			bLhsActive = bLhsActiveNext;
+			bRhsActive = bRhsActiveNext;
+			LhsIndex++; RhsIndex++;
+		}
+	}
+
+	// Process any remaining lhs events
+	while (LhsIndex < LhsNum * 2)
+	{
+		uniform const bool bLhsActiveNext = LhsIndex % 2 == 0;
+		uniform const int32 LhsT = bLhsActiveNext ? LhsStarts[LhsIndex / 2] : LhsStarts[LhsIndex / 2] + LhsLengths[LhsIndex / 2];
+
+		// Activate output
+		if (!bOutActive && bLhsActiveNext)
+		{
+			bOutActive = true;
+			OutStarts[OutIndex] = LhsT;
+		}
+		// Deactivate output
+		else if (bOutActive && !bLhsActiveNext)
+		{
+			bOutActive = false;
+			OutLengths[OutIndex] = LhsT - OutStarts[OutIndex];
+			OutIndex++;
+		}
+
+		bLhsActive = bLhsActiveNext;
+		LhsIndex++;
+	}
+
+	// Return number of ranges added
+	return OutIndex;
+}
+
+
+export uniform int32 LearningRangesUnion(
+	uniform int32 OutStarts[],
+	uniform int32 OutLengths[],
+	uniform const int32 LhsStarts[],
+	uniform const int32 LhsLengths[],
+	uniform const int32 RhsStarts[],
+	uniform const int32 RhsLengths[],
+	uniform const int LhsNum,
+	uniform const int RhsNum)
+{
+	if (LhsNum == 0)
+	{
+		LearningRangesCopy(
+			OutStarts,
+			OutLengths,
+			RhsStarts,
+			RhsLengths,
+			RhsNum);
+
+		return RhsNum;
+	}
+
+	if (RhsNum == 0)
+	{
+		LearningRangesCopy(
+			OutStarts,
+			OutLengths,
+			LhsStarts,
+			LhsLengths,
+			LhsNum);
+
+		return LhsNum;
+	}
+
+	if (LearningRangesEqual(LhsStarts, LhsLengths, RhsStarts, RhsLengths, LhsNum, RhsNum))
+	{
+		LearningRangesCopy(
+			OutStarts,
+			OutLengths,
+			LhsStarts,
+			LhsLengths,
+			LhsNum);
+
+		return LhsNum;
+	}
+
+	// Activation state of each list of ranges
+	uniform bool bOutActive = false;
+	uniform bool bLhsActive = false;
+	uniform bool bRhsActive = false;
+
+	// Event index for each list of ranges
+	uniform int32 LhsIndex = 0;
+	uniform int32 RhsIndex = 0;
+	uniform int32 OutIndex = 0;
+
+	// While both ranges have events to process
+	while (LhsIndex < LhsNum * 2 && RhsIndex < RhsNum * 2)
+	{
+		// Are the next lhs, and rhs events active or inactive
+		uniform const bool bLhsActiveNext = LhsIndex % 2 == 0;
+		uniform const bool bRhsActiveNext = RhsIndex % 2 == 0;
+
+		// Time of the next lhs, and rhs events
+		uniform const int32 LhsT = bLhsActiveNext ? LhsStarts[LhsIndex / 2] : LhsStarts[LhsIndex / 2] + LhsLengths[LhsIndex / 2];
+		uniform const int32 RhsT = bRhsActiveNext ? RhsStarts[RhsIndex / 2] : RhsStarts[RhsIndex / 2] + RhsLengths[RhsIndex / 2];
+
+		// Event from lhs is coming first
+		if (LhsT < RhsT)
+		{
+			// Activate output
+			if (!bOutActive && bLhsActiveNext)
+			{
+				bOutActive = true;
+				OutStarts[OutIndex] = LhsT;
+			}
+			// Deactivate output
+			else if (bOutActive && !bLhsActiveNext && !bRhsActive)
+			{
+
+				bOutActive = false;
+				OutLengths[OutIndex] = LhsT - OutStarts[OutIndex];
+				OutIndex++;
+			}
+
+			bLhsActive = bLhsActiveNext;
+			LhsIndex++;
+		}
+		// Event from rhs is coming first
+		else if (RhsT < LhsT)
+		{
+			// Activate output
+			if (!bOutActive && bRhsActiveNext)
+			{
+				bOutActive = true;
+				OutStarts[OutIndex] = RhsT;
+			}
+			// Deactivate output
+			else if (bOutActive && !bLhsActive && !bRhsActiveNext)
+			{
+				bOutActive = false;
+				OutLengths[OutIndex] = RhsT - OutStarts[OutIndex];
+				OutIndex++;
+			}
+
+			bRhsActive = bRhsActiveNext;
+			RhsIndex++;
+		}
+		// Event from lhs and rhs coming at same time
+		else
+		{
+			// Activate output
+			if (!bOutActive && (bLhsActiveNext || bRhsActiveNext))
+			{
+				bOutActive = true;
+				OutStarts[OutIndex] = LhsT;
+			}
+			// Deactivate output
+			else if (bOutActive && !(bLhsActiveNext || bRhsActiveNext))
+			{
+				bOutActive = false;
+				OutLengths[OutIndex] = LhsT - OutStarts[OutIndex];
+				OutIndex++;
+			}
+
+			bLhsActive = bLhsActiveNext;
+			bRhsActive = bRhsActiveNext;
+			LhsIndex++; RhsIndex++;
+		}
+	}
+
+	// Process any remaining lhs events
+	while (LhsIndex < LhsNum * 2)
+	{
+		uniform const bool bLhsActiveNext = LhsIndex % 2 == 0;
+		uniform const int32 LhsT = bLhsActiveNext ? LhsStarts[LhsIndex / 2] : LhsStarts[LhsIndex / 2] + LhsLengths[LhsIndex / 2];
+
+		// Activate output
+		if (!bOutActive && bLhsActiveNext)
+		{
+			bOutActive = true;
+			OutStarts[OutIndex] = LhsT;
+		}
+		// Deactivate output
+		else if (bOutActive && !bLhsActiveNext)
+		{
+			bOutActive = false;
+			OutLengths[OutIndex] = LhsT - OutStarts[OutIndex];
+			OutIndex++;
+		}
+
+		bLhsActive = bLhsActiveNext;
+		LhsIndex++;
+	}
+
+	// Process any remaining rhs events
+	while (RhsIndex < RhsNum * 2)
+	{
+		uniform const bool bRhsActiveNext = RhsIndex % 2 == 0;
+		uniform const int32 RhsT = bRhsActiveNext ? RhsStarts[RhsIndex / 2] : RhsStarts[RhsIndex / 2] + RhsLengths[RhsIndex / 2];
+
+		// Activate output
+		if (!bOutActive && bRhsActiveNext)
+		{
+			bOutActive = true;
+			OutStarts[OutIndex] = RhsT;
+		}
+		// Deactivate output
+		else if (bOutActive && !bRhsActiveNext)
+		{
+			bOutActive = false;
+			OutLengths[OutIndex] = RhsT - OutStarts[OutIndex];
+			OutIndex++;
+		}
+
+		bRhsActive = bRhsActiveNext;
+		RhsIndex++;
+	}
+
+	// Return number of ranges added
+	return OutIndex;
+}
+
+export uniform int32 LearningRangesIntersection(
+	uniform int32 OutStarts[],
+	uniform int32 OutLengths[],
+	uniform const int32 LhsStarts[],
+	uniform const int32 LhsLengths[],
+	uniform const int32 RhsStarts[],
+	uniform const int32 RhsLengths[],
+	uniform const int LhsNum,
+	uniform const int RhsNum)
+{
+	if (LhsNum == 0) { return 0; }
+	if (RhsNum == 0) { return 0; }
+
+	if (LearningRangesEqual(LhsStarts, LhsLengths, RhsStarts, RhsLengths, LhsNum, RhsNum))
+	{
+		LearningRangesCopy(
+			OutStarts,
+			OutLengths,
+			LhsStarts,
+			LhsLengths,
+			LhsNum);
+
+		return LhsNum;
+	}
+
+	// Activation state of each list of ranges
+	uniform bool bOutActive = false;
+	uniform bool bLhsActive = false;
+	uniform bool bRhsActive = false;
+
+	// Event index for each list of ranges
+	uniform int32 LhsIndex = 0;
+	uniform int32 RhsIndex = 0;
+	uniform int32 OutIndex = 0;
+
+	// While both ranges have events to process
+	while (LhsIndex < LhsNum * 2 && RhsIndex < RhsNum * 2)
+	{
+		// Are the next lhs, and rhs events active or inactive
+		uniform const bool bLhsActiveNext = LhsIndex % 2 == 0;
+		uniform const bool bRhsActiveNext = RhsIndex % 2 == 0;
+
+		// Time of the next lhs, and rhs events
+		uniform const int32 LhsT = bLhsActiveNext ? LhsStarts[LhsIndex / 2] : LhsStarts[LhsIndex / 2] + LhsLengths[LhsIndex / 2];
+		uniform const int32 RhsT = bRhsActiveNext ? RhsStarts[RhsIndex / 2] : RhsStarts[RhsIndex / 2] + RhsLengths[RhsIndex / 2];
+
+		// Event from lhs coming first
+		if (LhsT < RhsT)
+		{
+			// Activate output
+			if (!bOutActive && bRhsActive && bLhsActiveNext)
+			{
+				bOutActive = true;
+				OutStarts[OutIndex] = LhsT;
+			}
+			// Deactivate output
+			else if (bOutActive && !bLhsActiveNext)
+			{
+				bOutActive = false;
+				OutLengths[OutIndex] = LhsT - OutStarts[OutIndex];
+				OutIndex++;
+			}
+
+			bLhsActive = bLhsActiveNext;
+			LhsIndex++;
+		}
+		// Event from rhs coming first
+		else if (RhsT < LhsT)
+		{
+			// Activate output
+			if (!bOutActive && bLhsActive && bRhsActiveNext)
+			{
+				bOutActive = true;
+				OutStarts[OutIndex] = RhsT;
+			}
+			// Deactivate output
+			else if (bOutActive && !bRhsActiveNext)
+			{
+				bOutActive = false;
+				OutLengths[OutIndex] = RhsT - OutStarts[OutIndex];
+				OutIndex++;
+			}
+
+			bRhsActive = bRhsActiveNext;
+			RhsIndex++;
+		}
+		// Event from lhs and rhs coming at same time
+		else
+		{
+			// Activate output
+			if (!bOutActive && (bLhsActiveNext && bRhsActiveNext))
+			{
+				bOutActive = true;
+				OutStarts[OutIndex] = LhsT;
+			}
+			// Deactivate output
+			else if (bOutActive && (!bLhsActiveNext || !bRhsActiveNext))
+			{
+				bOutActive = false;
+				OutLengths[OutIndex] = LhsT - OutStarts[OutIndex];
+				OutIndex++;
+			}
+
+			bLhsActive = bLhsActiveNext;
+			bRhsActive = bRhsActiveNext;
+			LhsIndex++; RhsIndex++;
+		}
+	}
+
+	// Return number of ranges added
+	return OutIndex;
+}
+
+export uniform int32 LearningRangesTotalFrameNum(uniform const int32 RangeLengths[], uniform const int RangeNum)
+{
+	int32 Total = 0;
+
+	foreach (RangeIdx = 0 ... RangeNum)
+	{
+		Total += RangeLengths[RangeIdx];
+	}
+
+	return reduce_add(Total);
+}
+
+export uniform int32 LearningRangesIntersectionWithOffsets(
+	uniform int32 OutStarts[],
+	uniform int32 OutLengths[],
+	uniform int32 OutLhsOffsets[],
+	uniform int32 OutRhsOffsets[],
+	uniform const int32 LhsStarts[],
+	uniform const int32 LhsLengths[],
+	uniform const int32 RhsStarts[],
+	uniform const int32 RhsLengths[],
+	uniform const int32 LhsEntryOffset,
+	uniform const int32 RhsEntryOffset,
+	uniform const int32 LhsNum,
+	uniform const int32 RhsNum)
+{
+	if (LhsNum == 0) { return 0; }
+	if (RhsNum == 0) { return 0; }
+
+	if (LearningRangesEqual(LhsStarts, LhsLengths, RhsStarts, RhsLengths, LhsNum, RhsNum))
+	{
+		LearningRangesCopy(
+			OutStarts,
+			OutLengths,
+			LhsStarts,
+			LhsLengths,
+			LhsNum);
+
+		LearningComputeRangeOffsets(OutLhsOffsets, LhsLengths, LhsNum, LhsEntryOffset);
+		LearningComputeRangeOffsets(OutRhsOffsets, RhsLengths, RhsNum, RhsEntryOffset);
+		return LhsNum;
+	}
+
+	// Activation state of each list of ranges
+	uniform bool bOutActive = false;
+	uniform bool bLhsActive = false;
+	uniform bool bRhsActive = false;
+
+	// Event index for each list of ranges
+	uniform int32 LhsIndex = 0;
+	uniform int32 RhsIndex = 0;
+	uniform int32 OutIndex = 0;
+
+	// Offsets
+	uniform int32 LhsOffset = LhsEntryOffset;
+	uniform int32 RhsOffset = RhsEntryOffset;
+	uniform int32 LhsActiveStart = INDEX_NONE;
+	uniform int32 RhsActiveStart = INDEX_NONE;
+	uniform int32 LhsActiveOffset = INDEX_NONE;
+	uniform int32 RhsActiveOffset = INDEX_NONE;
+
+	// While both ranges have events to process
+	while (LhsIndex < LhsNum * 2 && RhsIndex < RhsNum * 2)
+	{
+		// Are the next lhs, and rhs events active or inactive
+		uniform const bool bLhsActiveNext = LhsIndex % 2 == 0;
+		uniform const bool bRhsActiveNext = RhsIndex % 2 == 0;
+
+		// Time of the next lhs, and rhs events
+		uniform const int32 LhsT = bLhsActiveNext ? LhsStarts[LhsIndex / 2] : LhsStarts[LhsIndex / 2] + LhsLengths[LhsIndex / 2];
+		uniform const int32 RhsT = bRhsActiveNext ? RhsStarts[RhsIndex / 2] : RhsStarts[RhsIndex / 2] + RhsLengths[RhsIndex / 2];
+
+		// Event from lhs coming first
+		if (LhsT < RhsT)
+		{
+			// Activate output
+			if (!bOutActive && bRhsActive && bLhsActiveNext)
+			{
+				bOutActive = true;
+				OutStarts[OutIndex] = LhsT;
+			}
+			// Deactivate output
+			else if (bOutActive && !bLhsActiveNext)
+			{
+				bOutActive = false;
+				OutLengths[OutIndex] = LhsT - OutStarts[OutIndex];
+				OutLhsOffsets[OutIndex] = LhsActiveOffset + (OutStarts[OutIndex] - LhsActiveStart);
+				OutRhsOffsets[OutIndex] = RhsActiveOffset + (OutStarts[OutIndex] - RhsActiveStart);
+				OutIndex++;
+			}
+
+			bLhsActive = bLhsActiveNext;
+			LhsActiveStart = bLhsActive ? LhsStarts[LhsIndex / 2] : INDEX_NONE;
+			LhsActiveOffset = bLhsActive ? LhsOffset : INDEX_NONE;
+			LhsOffset += bLhsActive ? 0 : LhsLengths[LhsIndex / 2];
+			LhsIndex++;
+		}
+		// Event from rhs coming first
+		else if (RhsT < LhsT)
+		{
+			// Activate output
+			if (!bOutActive && bLhsActive && bRhsActiveNext)
+			{
+				bOutActive = true;
+				OutStarts[OutIndex] = RhsT;
+			}
+			// Deactivate output
+			else if (bOutActive && !bRhsActiveNext)
+			{
+				bOutActive = false;
+				OutLengths[OutIndex] = RhsT - OutStarts[OutIndex];
+				OutLhsOffsets[OutIndex] = LhsActiveOffset + (OutStarts[OutIndex] - LhsActiveStart);
+				OutRhsOffsets[OutIndex] = RhsActiveOffset + (OutStarts[OutIndex] - RhsActiveStart);
+				OutIndex++;
+			}
+
+			bRhsActive = bRhsActiveNext;
+			RhsActiveStart = bRhsActive ? RhsStarts[RhsIndex / 2] : INDEX_NONE;
+			RhsActiveOffset = bRhsActive ? RhsOffset : INDEX_NONE;
+			RhsOffset += bRhsActive ? 0 : RhsLengths[RhsIndex / 2];
+			RhsIndex++;
+		}
+		// Event from lhs and rhs coming at same time
+		else
+		{
+			// Activate output
+			if (!bOutActive && (bLhsActiveNext && bRhsActiveNext))
+			{
+				bOutActive = true;
+				OutStarts[OutIndex] = LhsT;
+			}
+			// Deactivate output
+			else if (bOutActive && (!bLhsActiveNext || !bRhsActiveNext))
+			{
+				bOutActive = false;
+				OutLengths[OutIndex] = LhsT - OutStarts[OutIndex];
+				OutLhsOffsets[OutIndex] = LhsActiveOffset + (OutStarts[OutIndex] - LhsActiveStart);
+				OutRhsOffsets[OutIndex] = RhsActiveOffset + (OutStarts[OutIndex] - RhsActiveStart);
+				OutIndex++;
+			}
+
+			bLhsActive = bLhsActiveNext;
+			LhsActiveStart = bLhsActive ? LhsStarts[LhsIndex / 2] : INDEX_NONE;
+			LhsActiveOffset = bLhsActive ? LhsOffset : INDEX_NONE;
+			bRhsActive = bRhsActiveNext;
+			RhsActiveStart = bRhsActive ? RhsStarts[RhsIndex / 2] : INDEX_NONE;
+			RhsActiveOffset = bRhsActive ? RhsOffset : INDEX_NONE;
+			LhsOffset += bLhsActive ? 0 : LhsLengths[LhsIndex / 2];
+			RhsOffset += bRhsActive ? 0 : RhsLengths[RhsIndex / 2];
+			LhsIndex++; RhsIndex++;
+		}
+	}
+
+	// Return number of ranges added
+	return OutIndex;
+}
+
+export uniform int32 LearningRangesDifference(
+	uniform int32 OutStarts[],
+	uniform int32 OutLengths[],
+	uniform const int32 LhsStarts[],
+	uniform const int32 LhsLengths[],
+	uniform const int32 RhsStarts[],
+	uniform const int32 RhsLengths[],
+	uniform const int32 LhsNum,
+	uniform const int32 RhsNum)
+{
+	if (LhsNum == 0) { return 0; }
+
+	if (RhsNum == 0)
+	{
+		LearningRangesCopy(
+			OutStarts,
+			OutLengths,
+			LhsStarts,
+			LhsLengths,
+			LhsNum);
+
+		return LhsNum;
+	}
+
+	if (LearningRangesEqual(LhsStarts, LhsLengths, RhsStarts, RhsLengths, LhsNum, RhsNum))
+	{
+		return 0;
+	}
+
+	// Activation state of each list of ranges
+	uniform bool bOutActive = false;
+	uniform bool bLhsActive = false;
+	uniform bool bRhsActive = false;
+
+	// Event index for each list of ranges
+	uniform int32 LhsIndex = 0;
+	uniform int32 RhsIndex = 0;
+	uniform int32 OutIndex = 0;
+
+	// While both ranges have events to process
+	while (LhsIndex < LhsNum * 2 && RhsIndex < RhsNum * 2)
+	{
+		// Are the next lhs, and rhs events active or inactive
+		uniform const bool bLhsActiveNext = LhsIndex % 2 == 0;
+		uniform const bool bRhsActiveNext = RhsIndex % 2 == 0;
+
+		// Time of the next lhs, and rhs events
+		uniform const int32 LhsT = bLhsActiveNext ? LhsStarts[LhsIndex / 2] : LhsStarts[LhsIndex / 2] + LhsLengths[LhsIndex / 2];
+		uniform const int32 RhsT = bRhsActiveNext ? RhsStarts[RhsIndex / 2] : RhsStarts[RhsIndex / 2] + RhsLengths[RhsIndex / 2];
+
+		// Event coming from lhs first
+		if (LhsT < RhsT)
+		{
+			// Activate output
+			if (!bOutActive && !bRhsActive && bLhsActiveNext)
+			{
+				bOutActive = true;
+				OutStarts[OutIndex] = LhsT;
+			}
+			// Deactivate output
+			else if (bOutActive && !bLhsActiveNext)
+			{
+				bOutActive = false;
+				OutLengths[OutIndex] = LhsT - OutStarts[OutIndex];
+				OutIndex++;
+			}
+
+			bLhsActive = bLhsActiveNext;
+			LhsIndex++;
+		}
+		// Event coming from rhs first
+		else if (RhsT < LhsT)
+		{
+			// Activate output
+			if (!bOutActive && bLhsActive && !bRhsActiveNext)
+			{
+				bOutActive = true;
+				OutStarts[OutIndex] = RhsT;
+			}
+			// Deactivate output
+			else if (bOutActive && bRhsActiveNext)
+			{
+				bOutActive = false;
+				OutLengths[OutIndex] = RhsT - OutStarts[OutIndex];
+				OutIndex++;
+			}
+
+			bRhsActive = bRhsActiveNext;
+			RhsIndex++;
+		}
+		// Event from lhs and rhs coming at same time
+		else
+		{
+			// Activate output
+			if (!bOutActive && bLhsActiveNext && !bRhsActiveNext)
+			{
+				bOutActive = true;
+				OutStarts[OutIndex] = LhsT;
+			}
+			// Deactivate output
+			else if (bOutActive && bRhsActiveNext)
+			{
+				bOutActive = false;
+				OutLengths[OutIndex] = LhsT - OutStarts[OutIndex];
+				OutIndex++;
+			}
+
+			bLhsActive = bLhsActiveNext;
+			bRhsActive = bRhsActiveNext;
+			LhsIndex++; RhsIndex++;
+		}
+	}
+
+	// Process any remaining lhs events
+	while (LhsIndex < LhsNum * 2)
+	{
+		uniform const bool bLhsActiveNext = LhsIndex % 2 == 0;
+		uniform const int32 LhsT = bLhsActiveNext ? LhsStarts[LhsIndex / 2] : LhsStarts[LhsIndex / 2] + LhsLengths[LhsIndex / 2];
+
+		// Activate output
+		if (!bOutActive && bLhsActiveNext)
+		{
+			bOutActive = true;
+			OutStarts[OutIndex] = LhsT;
+		}
+		// Deactivate output
+		else if (bOutActive && !bLhsActiveNext)
+		{
+			bOutActive = false;
+			OutLengths[OutIndex] = LhsT - OutStarts[OutIndex];
+			OutIndex++;
+		}
+
+		bLhsActive = bLhsActiveNext;
+		LhsIndex++;
+	}
+
+	// Return number of ranges added
+	return OutIndex;
+}
+
 export void LearningFilterGaussian(
 	uniform float Out[], 
 	uniform const float In[], 
@@ -853,6 +2106,50 @@ export void LearningFilterMajorityVote(
 	}
 }
 
+export void LearningFrameAttributeZero(
+	uniform float Out[],
+	uniform int ValueNum)
+{
+	foreach(ValueIdx = 0 ... ValueNum)
+	{
+		Out[ValueIdx] = 0.0f;
+	}
+}
+
+export void LearningFrameAttributeSet(
+	uniform float Out[],
+	uniform int ValueNum,
+	uniform float Value)
+{
+	foreach(ValueIdx = 0 ... ValueNum)
+	{
+		Out[ValueIdx] = Value;
+	}
+}
+
+export void LearningFrameAttributeCopy(
+	uniform float Out[],
+	uniform const float In[],
+	uniform int ValueNum)
+{
+	foreach(ValueIdx = 0 ... ValueNum)
+	{
+		Out[ValueIdx] = In[ValueIdx];
+	}
+}
+
+export uniform bool LearningFrameAttributeEqual(
+	uniform const float Lhs[],
+	uniform const float Rhs[],
+	uniform int ValueNum)
+{
+	for (uniform int ValueIdx = 0; ValueIdx < ValueNum; ValueIdx++)
+	{
+		if (Lhs[ValueIdx] != Rhs[ValueIdx]) { return false; }
+	}
+	return true;
+}
+
 export void LearningFrameAttributeSelect(
 	uniform float Out[], 
 	uniform const float Cond[],
@@ -1002,6 +2299,36 @@ export void LearningFrameAttributeSqrtInplace(
 	}
 }
 
+export void LearningFrameAttributeLengthSquared(
+	uniform float InOut[], 
+	uniform const float In[], 
+	uniform const int ValueNum)
+{
+	foreach(ValueIdx = 0 ... ValueNum)
+	{
+		InOut[ValueIdx] += In[ValueIdx] * In[ValueIdx];
+	}
+}
+
+export void LearningInertializationDistanceConstant(
+	uniform float OutValues[],
+	uniform const float Locations[],
+	uniform const float Velocities[],
+	uniform const int ValueNum,
+	uniform const float Location,
+	uniform const float Velocity,
+	uniform const float BlendTime,
+	uniform const float LocWeight,
+	uniform const float VelWeight)
+{
+	foreach (ValueIdx = 0 ... ValueNum)
+	{
+		const float Integral = LearningInertializationIntegralMagnitude(Locations[ValueIdx] - Location, Velocities[ValueIdx] - Velocity, BlendTime);
+		const float MaxVelocity = LearningInertializationMaximumVelocityMagnitude(Locations[ValueIdx] - Location, Velocities[ValueIdx] - Velocity, BlendTime);
+		OutValues[ValueIdx] = LocWeight * Integral + VelWeight * MaxVelocity;
+	}
+}
+
 export uniform bool LearningFrameAttributeFindMinimum(
 	uniform int32& OutChannelIdx,
 	uniform int32& OutFrameIdx,
diff --git a/Engine/Plugins/Experimental/LearningCore/Source/Learning/Private/LearningFrameAttribute.cpp b/Engine/Plugins/Experimental/LearningCore/Source/Learning/Private/LearningFrameAttribute.cpp
--- a/Engine/Plugins/Experimental/LearningCore/Source/Learning/Private/LearningFrameAttribute.cpp
+++ b/Engine/Plugins/Experimental/LearningCore/Source/Learning/Private/LearningFrameAttribute.cpp
@@ -20,6 +20,116 @@ namespace UE::Learning
 			return X >= 0.0 ? Q : UE_PI - Q;
 		}
 
+		static inline float InertializationPartialIntegral(
+			const float X,
+			const float V,
+			const float Time,
+			const float BlendTime)
+		{
+			const float Z = FMath::Max(BlendTime, UE_SMALL_NUMBER);
+			const float T = FMath::Clamp(Time / Z, 0.0f, 1.0f);
+			return
+				X * ((Z / 2) * T * T * T * T - Z * T * T * T + Z * T) +
+				V * (((Z * Z) / 4) * T * T * T * T - ((2 * Z * Z) / 3) * T * T * T + ((Z * Z) / 2) * T * T);
+		}
+
+		static inline float InertializationFullIntegral(
+			const float X,
+			const float V,
+			const float BlendTime)
+		{
+			const float Z = FMath::Max(BlendTime, UE_SMALL_NUMBER);
+			return X * (Z / 2) + V * ((Z * Z) / 12);
+		}
+
+		static inline float InertializationPosition(
+			const float X,
+			const float V,
+			const float Time,
+			const float BlendTime)
+		{
+			const float Z = FMath::Max(BlendTime, UE_SMALL_NUMBER);
+			const float T = FMath::Clamp(Time / Z, 0.0f, 1.0f);
+			return X * (2 * T * T * T - 3 * T * T + 1) + V * (Z * T * T * T - Z * 2 * T * T + Z * T);
+		}
+
+		static inline float InertializationVelocity(
+			const float X,
+			const float V,
+			const float Time,
+			const float BlendTime)
+		{
+			const float Z = FMath::Max(BlendTime, UE_SMALL_NUMBER);
+			const float T = FMath::Clamp(Time / Z, 0.0f, 1.0f);
+			return X * ((6 / Z) * T * T - (6 / Z) * T) + V * (3 * T * T - 4 * T + 1);
+		}
+
+		static inline float InertializationPositionIntersectionTime(
+			const float X,
+			const float V,
+			const float BlendTime)
+		{
+			const float Z = FMath::Max(BlendTime, UE_SMALL_NUMBER);
+			return (-Z * X) / (2 * X + Z * V);
+		}
+
+		static inline float InertializationVelocityIntersectionTime(
+			const float X,
+			const float V,
+			const float BlendTime)
+		{
+			const float Z = FMath::Max(BlendTime, UE_SMALL_NUMBER);
+			return (Z * Z * V) / (6 * X + 3 * Z * V);
+		}
+
+		static inline float InertializationAccelerationIntersectionTime(
+			const float X,
+			const float V,
+			const float BlendTime)
+		{
+			const float Z = FMath::Max(BlendTime, UE_SMALL_NUMBER);
+			return (3 * Z * X + 2 * Z * Z * V) / (6 * X + 3 * Z * V);
+		}
+
+		static inline float InertializationMaximumPositionMagnitude(
+			const float X,
+			const float V,
+			const float BlendTime)
+		{
+			const float Z = FMath::Max(BlendTime, UE_SMALL_NUMBER);
+			const float T = InertializationVelocityIntersectionTime(X, V, Z);
+
+			return T > 0.0f && T < Z ?
+				FMath::Max(FMath::Abs(X), FMath::Abs(InertializationPosition(X, V, T, Z))) : FMath::Abs(X);
+		}
+
+		static inline float InertializationMaximumVelocityMagnitude(
+			const float X,
+			const float V,
+			const float BlendTime)
+		{
+			const float Z = FMath::Max(BlendTime, UE_SMALL_NUMBER);
+			const float T = InertializationAccelerationIntersectionTime(X, V, Z);
+
+			return T > 0.0f && T < Z ?
+				FMath::Max(FMath::Abs(V), FMath::Abs(InertializationVelocity(X, V, T, Z))) : FMath::Abs(V);
+		}
+
+		static inline float InertializationIntegralMagnitude(
+			const float X,
+			const float V,
+			const float BlendTime)
+		{
+			const float Z = FMath::Max(BlendTime, UE_SMALL_NUMBER);
+			const float T = InertializationPositionIntersectionTime(X, V, Z);
+			const float IntegralT = InertializationPartialIntegral(X, V, T, Z);
+			const float IntegralZ = InertializationFullIntegral(X, V, Z);
+
+			return T > 0.0f && T < Z ?
+				FMath::Abs(IntegralZ - IntegralT) + FMath::Abs(IntegralT) :
+				FMath::Abs(IntegralZ);
+		}
+
 		static inline bool FindMinimum(
 			int32& OutChannelIdx,
 			int32& OutFrameIdx,
@@ -98,6 +208,63 @@ namespace UE::Learning
 #endif
 		}
 
+		static inline void Zero(TLearningArrayView<1, float> Out)
+		{
+#if UE_LEARNING_ISPC
+			ispc::LearningFrameAttributeZero(Out.GetData(), Out.Num());
+#else
+			const int32 ValueNum = Out.Num();
+			for (int32 ValueIdx = 0; ValueIdx < ValueNum; ValueIdx++)
+			{
+				Out[ValueIdx] = 0.0f;
+			}
+#endif
+		}
+
+		static inline void Set(TLearningArrayView<1, float> Out, const float Value)
+		{
+#if UE_LEARNING_ISPC
+			ispc::LearningFrameAttributeSet(Out.GetData(), Out.Num(), Value);
+#else
+			const int32 ValueNum = Out.Num();
+			for (int32 ValueIdx = 0; ValueIdx < ValueNum; ValueIdx++)
+			{
+				Out[ValueIdx] = Value;
+			}
+#endif
+		}
+
+		static inline void Copy(TLearningArrayView<1, float> Out, const TLearningArrayView<1, const float> In)
+		{
+			check(Out.Num() == In.Num());
+
+#if UE_LEARNING_ISPC
+			ispc::LearningFrameAttributeCopy(Out.GetData(), In.GetData(), Out.Num());
+#else
+			const int32 ValueNum = Out.Num();
+			for (int32 ValueIdx = 0; ValueIdx < ValueNum; ValueIdx++)
+			{
+				Out[ValueIdx] = In[ValueIdx];
+			}
+#endif
+		}
+
+		static inline bool Equal(const TLearningArrayView<1, const float> Lhs, const TLearningArrayView<1, const float> Rhs)
+		{
+			check(Lhs.Num() == Rhs.Num());
+
+#if UE_LEARNING_ISPC
+			return ispc::LearningFrameAttributeEqual(Lhs.GetData(), Rhs.GetData(), Lhs.Num());
+#else
+			const int32 ValueNum = Lhs.Num();
+			for (int32 ValueIdx = 0; ValueIdx < ValueNum; ValueIdx++)
+			{
+				if (Lhs[ValueIdx] != Rhs[ValueIdx]) { return false; }
+			}
+			return true;
+#endif
+		}
+
 		static inline void Select(TLearningArrayView<1, float> Out, const TLearningArrayView<1, const float> Cond, const TLearningArrayView<1, const float> True, const TLearningArrayView<1, const float> False)
 		{
 			check(Out.Num() == Cond.Num());
@@ -388,11 +555,19 @@ namespace UE::Learning
 		static inline void LengthSquared(TLearningArrayView<1, float> InOut, const TLearningArrayView<1, const float> In)
 		{
 			check(InOut.Num() == In.Num());
+
+#if UE_LEARNING_ISPC
+			ispc::LearningFrameAttributeLengthSquared(
+				InOut.GetData(),
+				In.GetData(),
+				InOut.Num());
+#else
 			const int32 ValueNum = InOut.Num();
 			for (int32 ValueIdx = 0; ValueIdx < ValueNum; ValueIdx++)
 			{
 				InOut[ValueIdx] += In[ValueIdx] * In[ValueIdx];
 			}
+#endif
 		}
 
 		static inline void Normalize(TLearningArrayView<2, float> Out, const TLearningArrayView<2, const float> In)
@@ -795,7 +970,7 @@ namespace UE::Learning
 			const int32 ValueNum = Out.Num();
 			for (int32 ValueIdx = 0; ValueIdx < ValueNum; ValueIdx++)
 			{
-				Out[ValueIdx] = Lhs < Rhs[ValueIdx];
+				Out[ValueIdx] = Lhs <= Rhs[ValueIdx];
 			}
 		}
 
@@ -1389,6 +1564,43 @@ namespace UE::Learning
 			}
 		}
 
+		static inline void InertializationDistanceConstant(
+			TLearningArrayView<1, float> OutValues,
+			const TLearningArrayView<1, const float> Locations,
+			const TLearningArrayView<1, const float> Velocities,
+			const float Location,
+			const float Velocity,
+			const float BlendTime,
+			const float LocWeight,
+			const float VelWeight)
+		{
+			check(OutValues.Num() == Locations.Num());
+			check(OutValues.Num() == Velocities.Num());
+
+#if UE_LEARNING_ISPC
+			ispc::LearningInertializationDistanceConstant(
+				OutValues.GetData(),
+				Locations.GetData(),
+				Velocities.GetData(),
+				OutValues.Num(),
+				Location,
+				Velocity,
+				BlendTime,
+				LocWeight,
+				VelWeight);
+#else
+			const int32 ValueNum = OutValues.Num();
+
+			for (int32 ValueIdx = 0; ValueIdx < ValueNum; ValueIdx++)
+			{
+				const float Integral = InertializationIntegralMagnitude(Locations[ValueIdx] - Location, Velocities[ValueIdx] - Velocity, BlendTime);
+				const float MaxVelocity = InertializationMaximumVelocityMagnitude(Locations[ValueIdx] - Location, Velocities[ValueIdx] - Velocity, BlendTime);
+
+				OutValues[ValueIdx] = LocWeight * Integral + VelWeight * MaxVelocity;
+			}
+#endif
+		}
+
 		// We need these custom Vector4 functions because the defaults for the FVector4f class will treat it like a Vector3
 
 		static inline float Vector4Length(const FVector4f X)
@@ -1850,18 +2062,18 @@ namespace UE::Learning
 			const TLearningArrayView<1, const float> InRotZ,
 			const TLearningArrayView<1, const float> InRotW)
 		{
-			Array::Copy(OutPosX, InPosX);
-			Array::Copy(OutPosY, InPosY);
-			Array::Copy(OutPosZ, InPosZ);
+			Private::Copy(OutPosX, InPosX);
+			Private::Copy(OutPosY, InPosY);
+			Private::Copy(OutPosZ, InPosZ);
 
-			Array::Copy(OutRotX, InRotX);
-			Array::Copy(OutRotY, InRotY);
-			Array::Copy(OutRotZ, InRotZ);
-			Array::Copy(OutRotW, InRotW);
+			Private::Copy(OutRotX, InRotX);
+			Private::Copy(OutRotY, InRotY);
+			Private::Copy(OutRotZ, InRotZ);
+			Private::Copy(OutRotW, InRotW);
 
-			Array::Set(OutSclX, 1.0f);
-			Array::Set(OutSclY, 1.0f);
-			Array::Set(OutSclZ, 1.0f);
+			Private::Set(OutSclX, 1.0f);
+			Private::Set(OutSclY, 1.0f);
+			Private::Set(OutSclZ, 1.0f);
 		}
 
 		static inline void TransformMake(
@@ -1886,18 +2098,18 @@ namespace UE::Learning
 			const TLearningArrayView<1, const float> InSclY,
 			const TLearningArrayView<1, const float> InSclZ)
 		{
-			Array::Copy(OutPosX, InPosX);
-			Array::Copy(OutPosY, InPosY);
-			Array::Copy(OutPosZ, InPosZ);
+			Private::Copy(OutPosX, InPosX);
+			Private::Copy(OutPosY, InPosY);
+			Private::Copy(OutPosZ, InPosZ);
 
-			Array::Copy(OutRotX, InRotX);
-			Array::Copy(OutRotY, InRotY);
-			Array::Copy(OutRotZ, InRotZ);
-			Array::Copy(OutRotW, InRotW);
+			Private::Copy(OutRotX, InRotX);
+			Private::Copy(OutRotY, InRotY);
+			Private::Copy(OutRotZ, InRotZ);
+			Private::Copy(OutRotW, InRotW);
 
-			Array::Copy(OutSclX, InSclX);
-			Array::Copy(OutSclY, InSclY);
-			Array::Copy(OutSclZ, InSclZ);
+			Private::Copy(OutSclX, InSclX);
+			Private::Copy(OutSclY, InSclY);
+			Private::Copy(OutSclZ, InSclZ);
 		}
 
 	}
@@ -1952,7 +2164,7 @@ namespace UE::Learning
 	{
 		bool Equal(const FFrameAttribute& Lhs, const FFrameAttribute& Rhs)
 		{
-			return FrameRangeSet::Equal(Lhs.FrameRangeSet, Rhs.FrameRangeSet) && Array::Equal(Lhs.AttributeData, Rhs.AttributeData);
+			return FrameRangeSet::Equal(Lhs.FrameRangeSet, Rhs.FrameRangeSet) && Private::Equal(Lhs.AttributeData.Flatten(), Rhs.AttributeData.Flatten());
 		}
 
 		void Intersection(FFrameAttribute& OutFrameAttribute, const FFrameAttribute& FrameAttribute, const FFrameRangeSet& FrameRangeSet)
@@ -1993,7 +2205,7 @@ namespace UE::Learning
 					const int32 LhsOffset = LhsOffsets[RangeIdx];
 					const int32 Length = OutFrameAttribute.FrameRangeSet.GetAllRangeLengths()[RangeIdx];
 
-					Array::Copy(
+					Private::Copy(
 						OutFrameAttribute.GetChannelRangeAttributeData(ChannelIdx, OutOffset, Length),
 						FrameAttribute.GetChannelRangeAttributeData(ChannelIdx, LhsOffset, Length));
 				}
@@ -2158,6 +2370,7 @@ namespace UE::Learning
 			{
 				Out.Empty();
 				Op(Out, {}, {}, {}, {});
+				return;
 			}
 
 			// Check All Equal
@@ -2252,14 +2465,14 @@ namespace UE::Learning
 		{
 			Out.FrameRangeSet = FrameRangeSet;
 			Out.AttributeData.SetNumUninitialized({ ChannelNum, FrameRangeSet.GetTotalFrameNum() });
-			Array::Zero(Out.AttributeData);
+			Private::Zero(Out.AttributeData.Flatten());
 		}
 
 		void Ones(FFrameAttribute& Out, const FFrameRangeSet& FrameRangeSet, const int32 ChannelNum)
 		{
 			Out.FrameRangeSet = FrameRangeSet;
 			Out.AttributeData.SetNumUninitialized({ ChannelNum, FrameRangeSet.GetTotalFrameNum() });
-			Array::Set(Out.AttributeData, 1.0f);
+			Private::Set(Out.AttributeData.Flatten(), 1.0f);
 		}
 
 		void Fill(FFrameAttribute& Out, const FFrameRangeSet& FrameRangeSet, const TLearningArrayView<1, const float> Values)
@@ -2270,7 +2483,7 @@ namespace UE::Learning
 			Out.AttributeData.SetNumUninitialized({ ChannelNum, FrameRangeSet.GetTotalFrameNum() });
 			for (int32 ChannelIdx = 0; ChannelIdx < ChannelNum; ChannelIdx++)
 			{
-				Array::Set(Out.AttributeData[ChannelIdx], Values[ChannelIdx]);
+				Private::Set(Out.AttributeData[ChannelIdx], Values[ChannelIdx]);
 			}
 		}
 
@@ -2451,7 +2664,7 @@ namespace UE::Learning
 				const TLearningArrayView<1, const int32> RhsRangeOffsets,
 				const TLearningArrayView<1, const int32> RangeLengths) {
 
-					Array::Zero(Out.GetAttributeData());
+					Private::Zero(Out.GetAttributeData().Flatten());
 
 					const int32 ChannelNum = Out.GetChannelNum();
 					const int32 RangeNum = RangeLengths.Num();
@@ -2477,7 +2690,7 @@ namespace UE::Learning
 				const TLearningArrayView<1, const int32> RangeOffsets,
 				const TLearningArrayView<1, const int32> RangeLengths) {
 
-					Array::Copy(Out.AttributeData, In.AttributeData);
+					Private::Copy(Out.AttributeData.Flatten(), In.AttributeData.Flatten());
 				});
 		}
 
@@ -2755,7 +2968,7 @@ namespace UE::Learning
 
 					const int32 ChannelNum = In.GetChannelNum();
 
-					Array::Zero(Out.GetAttributeData());
+					Private::Zero(Out.GetAttributeData().Flatten());
 
 					for (int32 ChannelIdx = 0; ChannelIdx < ChannelNum; ChannelIdx++)
 					{
@@ -2792,7 +3005,7 @@ namespace UE::Learning
 				const TLearningArrayView<1, const int32> RangeOffsets,
 				const TLearningArrayView<1, const int32> RangeLengths) {
 
-					Array::Copy(Out.GetChannelAttributeData(0), In.GetChannelAttributeData(ChannelIdx));
+					Private::Copy(Out.GetChannelAttributeData(0), In.GetChannelAttributeData(ChannelIdx));
 
 				});
 		}
@@ -2811,7 +3024,7 @@ namespace UE::Learning
 					check(Out.GetChannelNum() == ChannelNum);
 					for (int32 ChannelIdx = 0; ChannelIdx < ChannelNum; ChannelIdx++)
 					{
-						Array::Copy(Out.GetChannelAttributeData(ChannelIdx), In.GetChannelAttributeData(ChannelIdx + StartChannelIdx));
+						Private::Copy(Out.GetChannelAttributeData(ChannelIdx), In.GetChannelAttributeData(ChannelIdx + StartChannelIdx));
 					}
 
 				});
@@ -2987,7 +3200,7 @@ namespace UE::Learning
 
 					const int32 ChannelNum = In.GetChannelNum();
 
-					Array::Zero(Out.GetAttributeData());
+					Private::Zero(Out.GetAttributeData().Flatten());
 
 					for (int32 ChannelIdx = 0; ChannelIdx < ChannelNum; ChannelIdx++)
 					{
@@ -3009,7 +3222,7 @@ namespace UE::Learning
 
 					const int32 ChannelNum = In.GetChannelNum();
 
-					Array::Zero(Out.GetAttributeData());
+					Private::Zero(Out.GetAttributeData().Flatten());
 
 					for (int32 ChannelIdx = 0; ChannelIdx < ChannelNum; ChannelIdx++)
 					{
@@ -3031,7 +3244,7 @@ namespace UE::Learning
 
 					const int32 ChannelNum = In.GetChannelNum();
 
-					Array::Zero(Out.GetAttributeData());
+					Private::Zero(Out.GetAttributeData().Flatten());
 
 					for (int32 ChannelIdx = 0; ChannelIdx < ChannelNum; ChannelIdx++)
 					{
@@ -4601,5 +4814,71 @@ namespace UE::Learning
 				});
 		}
 
+		void InertializationDistanceConstant(
+			FFrameAttribute& Out,
+			const FFrameAttribute& LhsPos,
+			const FFrameAttribute& LhsVel,
+			const TArrayView<const float> RhsLoc,
+			const TArrayView<const float> RhsVel,
+			const float BlendTime,
+			const float LocWeight,
+			const float VelWeight)
+		{
+			check(LhsPos.GetChannelNum() == LhsVel.GetChannelNum());
+			check(LhsPos.GetChannelNum() == RhsLoc.Num());
+			check(LhsPos.GetChannelNum() == RhsVel.Num());
+
+			const int32 ChannelNum = LhsPos.GetChannelNum();
+
+			if (FrameRangeSet::Equal(LhsPos.FrameRangeSet, LhsVel.FrameRangeSet))
+			{
+				Out.FrameRangeSet = LhsPos.FrameRangeSet;
+				Out.AttributeData.SetNumUninitialized({ ChannelNum, LhsPos.GetTotalFrameNum() });
+
+				for (int32 ChannelIdx = 0; ChannelIdx < ChannelNum; ChannelIdx++)
+				{
+					Private::InertializationDistanceConstant(
+						Out.GetChannelAttributeData(ChannelIdx),
+						LhsPos.GetChannelAttributeData(ChannelIdx),
+						LhsVel.GetChannelAttributeData(ChannelIdx),
+						RhsLoc[ChannelIdx],
+						RhsVel[ChannelIdx],
+						BlendTime,
+						LocWeight,
+						VelWeight);
+				}
+
+				return;
+			}
+
+			BinaryOp(Out, ChannelNum, LhsPos, LhsVel, [RhsLoc, RhsVel, BlendTime, LocWeight, VelWeight](
+				FFrameAttribute& Out,
+				const FFrameAttribute& Lhs,
+				const FFrameAttribute& Rhs,
+				const TLearningArrayView<1, const int32> OutRangeOffsets,
+				const TLearningArrayView<1, const int32> LhsRangeOffsets,
+				const TLearningArrayView<1, const int32> RhsRangeOffsets,
+				const TLearningArrayView<1, const int32> RangeLengths) {
+
+					const int32 ChannelNum = Out.GetChannelNum();
+					const int32 RangeNum = RangeLengths.Num();
+
+					for (int32 ChannelIdx = 0; ChannelIdx < ChannelNum; ChannelIdx++)
+					{
+						for (int32 RangeIdx = 0; RangeIdx < RangeNum; RangeIdx++)
+						{
+							Private::InertializationDistanceConstant(
+								Out.GetChannelRangeAttributeData(ChannelIdx, OutRangeOffsets[RangeIdx], RangeLengths[RangeIdx]),
+								Lhs.GetChannelRangeAttributeData(ChannelIdx, LhsRangeOffsets[RangeIdx], RangeLengths[RangeIdx]),
+								Rhs.GetChannelRangeAttributeData(ChannelIdx, RhsRangeOffsets[RangeIdx], RangeLengths[RangeIdx]),
+								RhsLoc[ChannelIdx],
+								RhsVel[ChannelIdx],
+								BlendTime,
+								LocWeight,
+								VelWeight);
+						}
+					}
+				});
+		}
 	}
 }
diff --git a/Engine/Plugins/Experimental/LearningCore/Source/Learning/Private/LearningFrameRangeSet.cpp b/Engine/Plugins/Experimental/LearningCore/Source/Learning/Private/LearningFrameRangeSet.cpp
--- a/Engine/Plugins/Experimental/LearningCore/Source/Learning/Private/LearningFrameRangeSet.cpp
+++ b/Engine/Plugins/Experimental/LearningCore/Source/Learning/Private/LearningFrameRangeSet.cpp
@@ -59,6 +59,14 @@ namespace UE::Learning
 			const int32 InitialOffset = 0)
 		{
 			check(Offsets.Num() == Lengths.Num());
+
+#if UE_LEARNING_ISPC
+			return ispc::LearningComputeRangeOffsets(
+				Offsets.GetData(),
+				Lengths.GetData(),
+				Lengths.Num(),
+				InitialOffset);
+#else
 			const int32 RangeNum = Lengths.Num();
 			int32 Offset = InitialOffset;
 			for (int32 RangeIdx = 0; RangeIdx < RangeNum; RangeIdx++)
@@ -66,6 +74,80 @@ namespace UE::Learning
 				Offsets[RangeIdx] = Offset;
 				Offset += Lengths[RangeIdx];
 			}
+#endif
+		}
+
+		static inline void FramesCopy(
+			const TLearningArrayView<1, int32> DstFrames,
+			const TLearningArrayView<1, const int32> SrcFrames)
+		{
+			check(DstFrames.Num() == SrcFrames.Num());
+
+#if UE_LEARNING_ISPC
+			return ispc::LearningFramesCopy(
+				DstFrames.GetData(),
+				SrcFrames.GetData(),
+				DstFrames.Num());
+#else
+
+			const int32 FrameNum = DstFrames.Num();
+			for (int32 FrameIdx = 0; FrameIdx < FrameNum; FrameIdx++)
+			{
+				DstFrames[FrameIdx] = SrcFrames[FrameIdx];
+			}
+#endif
+		}
+
+		static inline void RangesCopy(
+			const TLearningArrayView<1, int32> DstStarts,
+			const TLearningArrayView<1, int32> DstLengths,
+			const TLearningArrayView<1, const int32> SrcStarts,
+			const TLearningArrayView<1, const int32> SrcLengths)
+		{
+			check(DstStarts.Num() == DstLengths.Num());
+			check(SrcStarts.Num() == SrcLengths.Num());
+			check(DstStarts.Num() == SrcStarts.Num());
+
+#if UE_LEARNING_ISPC
+			return ispc::LearningRangesCopy(
+				DstStarts.GetData(),
+				DstLengths.GetData(),
+				SrcStarts.GetData(),
+				SrcLengths.GetData(),
+				DstStarts.Num());
+#else
+
+			const int32 RangeNum = DstStarts.Num();
+			for (int32 RangeIdx = 0; RangeIdx < RangeNum; RangeIdx++)
+			{
+				DstStarts[RangeIdx] = SrcStarts[RangeIdx];
+				DstLengths[RangeIdx] = SrcLengths[RangeIdx];
+			}
+#endif
+		}
+
+		static inline void RangesCopyFrames(
+			const TLearningArrayView<1, int32> DstStarts,
+			const TLearningArrayView<1, int32> DstLengths,
+			const TLearningArrayView<1, const int32> SrcFrames)
+		{
+			check(DstStarts.Num() == DstLengths.Num());
+			check(DstStarts.Num() == SrcFrames.Num());
+
+#if UE_LEARNING_ISPC
+			return ispc::LearningRangesCopyFrames(
+				DstStarts.GetData(),
+				DstLengths.GetData(),
+				SrcFrames.GetData(),
+				DstStarts.Num());
+#else
+			const int32 RangeNum = DstStarts.Num();
+			for (int32 RangeIdx = 0; RangeIdx < RangeNum; RangeIdx++)
+			{
+				DstStarts[RangeIdx] = SrcFrames[RangeIdx];
+				DstLengths[RangeIdx] = 1;
+			}
+#endif
 		}
 
 		static inline bool RangesEqual(
@@ -191,6 +273,14 @@ namespace UE::Learning
 		{
 			check(Starts.Num() == Lengths.Num());
 
+#if UE_LEARNING_ISPC
+			return ispc::LearningRangesContainsTime(
+				Starts.GetData(),
+				Lengths.GetData(),
+				Starts.Num(),
+				Time,
+				FrameDeltaTime);
+#else
 			const int32 RangeNum = Starts.Num();
 			for (int32 RangeIdx = 0; RangeIdx < RangeNum; RangeIdx++)
 			{
@@ -203,6 +293,7 @@ namespace UE::Learning
 				}
 			}
 			return false;
+#endif
 		}
 
 		static inline bool RangesFind(
@@ -250,6 +341,16 @@ namespace UE::Learning
 		{
 			check(Starts.Num() == Lengths.Num());
 
+#if UE_LEARNING_ISPC
+			return ispc::LearningRangesFindTime(
+				OutRangeIdx,
+				OutRangeTime,
+				Starts.GetData(),
+				Lengths.GetData(),
+				Starts.Num(),
+				Time,
+				FrameDeltaTime);
+#else
 			const int32 RangeNum = Starts.Num();
 			for (int32 RangeIdx = 0; RangeIdx < RangeNum; RangeIdx++)
 			{
@@ -267,6 +368,7 @@ namespace UE::Learning
 			OutRangeIdx = INDEX_NONE;
 			OutRangeTime = -1.0f;
 			return false;
+#endif
 		}
 
 
@@ -279,20 +381,37 @@ namespace UE::Learning
 		{
 			TRACE_CPUPROFILER_EVENT_SCOPE(Learning::FrameRangeSet::Private::FramesRangesUnion);
 
+#if UE_LEARNING_ISPC
+			return ispc::LearningFramesRangesUnion(
+				OutStarts.GetData(),
+				OutLengths.GetData(),
+				LhsFrames.GetData(),
+				RhsStarts.GetData(),
+				RhsLengths.GetData(),
+				LhsFrames.Num(),
+				RhsStarts.Num());
+#else
 			FramesCheck(LhsFrames);
 			RangesCheck(RhsStarts, RhsLengths);
 
 			if (LhsFrames.IsEmpty())
 			{
-				Array::Copy(OutStarts.Slice(0, RhsStarts.Num()), RhsStarts);
-				Array::Copy(OutLengths.Slice(0, RhsLengths.Num()), RhsLengths);
+				RangesCopy(
+					OutStarts.Slice(0, RhsStarts.Num()), 
+					OutLengths.Slice(0, RhsLengths.Num()), 
+					RhsStarts,
+					RhsLengths);
+
 				return RhsLengths.Num();
 			}
 
 			if (RhsStarts.IsEmpty())
 			{
-				Array::Copy(OutStarts.Slice(0, LhsFrames.Num()), LhsFrames);
-				Array::Set(OutLengths.Slice(0, LhsFrames.Num()), 1);
+				RangesCopyFrames(
+					OutStarts.Slice(0, LhsFrames.Num()),
+					OutLengths.Slice(0, LhsFrames.Num()),
+					LhsFrames);
+
 				return LhsFrames.Num();
 			}
 
@@ -443,6 +562,7 @@ namespace UE::Learning
 
 			// Return number of ranges added
 			return OutIndex;
+#endif
 		}
 
 		static inline int32 FramesRangesIntersection(
@@ -453,6 +573,15 @@ namespace UE::Learning
 		{
 			TRACE_CPUPROFILER_EVENT_SCOPE(Learning::FrameRangeSet::Private::FramesRangesIntersection);
 
+#if UE_LEARNING_ISPC
+			return ispc::LearningFramesRangesIntersection(
+				OutFrames.GetData(),
+				LhsFrames.GetData(),
+				RhsStarts.GetData(),
+				RhsLengths.GetData(),
+				LhsFrames.Num(),
+				RhsStarts.Num());
+#else
 			FramesCheck(LhsFrames);
 			RangesCheck(RhsStarts, RhsLengths);
 
@@ -522,6 +651,7 @@ namespace UE::Learning
 
 			// Return number of frames added
 			return OutIndex;
+#endif
 		}
 
 		static inline int32 FramesRangesDifference(
@@ -532,6 +662,15 @@ namespace UE::Learning
 		{
 			TRACE_CPUPROFILER_EVENT_SCOPE(Learning::FrameRangeSet::Private::FramesRangesDifference);
 
+#if UE_LEARNING_ISPC
+			return ispc::LearningFramesRangesDifference(
+				OutFrames.GetData(),
+				LhsFrames.GetData(),
+				RhsStarts.GetData(),
+				RhsLengths.GetData(),
+				LhsFrames.Num(),
+				RhsStarts.Num());
+#else
 			FramesCheck(LhsFrames);
 			RangesCheck(RhsStarts, RhsLengths);
 
@@ -539,7 +678,7 @@ namespace UE::Learning
 
 			if (RhsStarts.IsEmpty())
 			{
-				Array::Copy(OutFrames.Slice(0, LhsFrames.Num()), LhsFrames);
+				FramesCopy(OutFrames.Slice(0, LhsFrames.Num()), LhsFrames);
 				return LhsFrames.Num();
 			}
 
@@ -606,6 +745,7 @@ namespace UE::Learning
 
 			// Return number of events added
 			return OutIndex;
+#endif
 		}
 
 
@@ -618,6 +758,16 @@ namespace UE::Learning
 		{
 			TRACE_CPUPROFILER_EVENT_SCOPE(Learning::FrameRangeSet::Private::RangesFramesDifference);
 
+#if UE_LEARNING_ISPC
+			return ispc::LearningRangesFramesDifference(
+				OutStarts.GetData(),
+				OutLengths.GetData(),
+				LhsStarts.GetData(),
+				LhsLengths.GetData(),
+				RhsFrames.GetData(),
+				LhsStarts.Num(),
+				RhsFrames.Num());
+#else
 			RangesCheck(LhsStarts, LhsLengths);
 			FramesCheck(RhsFrames);
 
@@ -625,8 +775,12 @@ namespace UE::Learning
 
 			if (RhsFrames.IsEmpty())
 			{
-				Array::Copy(OutStarts.Slice(0, LhsStarts.Num()), LhsStarts);
-				Array::Copy(OutLengths.Slice(0, LhsLengths.Num()), LhsLengths);
+				RangesCopy(
+					OutStarts.Slice(0, LhsStarts.Num()), 
+					OutLengths.Slice(0, LhsLengths.Num()),
+					LhsStarts,
+					LhsLengths);
+
 				return LhsLengths.Num();
 			}
 
@@ -748,6 +902,7 @@ namespace UE::Learning
 
 			// Return number of ranges added
 			return OutIndex;
+#endif
 		}
 
 
@@ -761,27 +916,50 @@ namespace UE::Learning
 		{
 			TRACE_CPUPROFILER_EVENT_SCOPE(Learning::FrameRangeSet::Private::RangesUnion);
 
+#if UE_LEARNING_ISPC
+			return ispc::LearningRangesUnion(
+				OutStarts.GetData(),
+				OutLengths.GetData(),
+				LhsStarts.GetData(),
+				LhsLengths.GetData(),
+				RhsStarts.GetData(),
+				RhsLengths.GetData(),
+				LhsStarts.Num(),
+				RhsStarts.Num());
+#else
 			RangesCheck(LhsStarts, LhsLengths);
 			RangesCheck(RhsStarts, RhsLengths);
 
 			if (LhsStarts.IsEmpty())
 			{
-				Array::Copy(OutStarts.Slice(0, RhsStarts.Num()), RhsStarts);
-				Array::Copy(OutLengths.Slice(0, RhsLengths.Num()), RhsLengths);
+				RangesCopy(
+					OutStarts.Slice(0, RhsStarts.Num()),
+					OutLengths.Slice(0, RhsLengths.Num()),
+					RhsStarts,
+					RhsLengths);
+
 				return RhsLengths.Num();
 			}
 
 			if (RhsStarts.IsEmpty())
 			{
-				Array::Copy(OutStarts.Slice(0, LhsStarts.Num()), LhsStarts);
-				Array::Copy(OutLengths.Slice(0, LhsLengths.Num()), LhsLengths);
+				RangesCopy(
+					OutStarts.Slice(0, LhsStarts.Num()),
+					OutLengths.Slice(0, LhsLengths.Num()),
+					LhsStarts,
+					LhsLengths);
+
 				return LhsLengths.Num();
 			}
 
 			if (RangesEqual(LhsStarts, LhsLengths, RhsStarts, RhsLengths))
 			{
-				Array::Copy(OutStarts.Slice(0, LhsStarts.Num()), LhsStarts);
-				Array::Copy(OutLengths.Slice(0, LhsLengths.Num()), LhsLengths);
+				RangesCopy(
+					OutStarts.Slice(0, LhsStarts.Num()),
+					OutLengths.Slice(0, LhsLengths.Num()),
+					LhsStarts,
+					LhsLengths);
+
 				return LhsLengths.Num();
 			}
 
@@ -932,6 +1110,7 @@ namespace UE::Learning
 
 			// Return number of ranges added
 			return OutIndex;
+#endif
 		}
 
 		static inline int32 RangesIntersection(
@@ -944,6 +1123,17 @@ namespace UE::Learning
 		{
 			TRACE_CPUPROFILER_EVENT_SCOPE(Learning::FrameRangeSet::Private::RangesIntersection);
 
+#if UE_LEARNING_ISPC
+			return ispc::LearningRangesIntersection(
+				OutStarts.GetData(),
+				OutLengths.GetData(),
+				LhsStarts.GetData(),
+				LhsLengths.GetData(),
+				RhsStarts.GetData(),
+				RhsLengths.GetData(),
+				LhsStarts.Num(),
+				RhsStarts.Num());
+#else
 			RangesCheck(LhsStarts, LhsLengths);
 			RangesCheck(RhsStarts, RhsLengths);
 
@@ -952,8 +1142,12 @@ namespace UE::Learning
 
 			if (RangesEqual(LhsStarts, LhsLengths, RhsStarts, RhsLengths))
 			{
-				Array::Copy(OutStarts.Slice(0, LhsStarts.Num()), LhsStarts);
-				Array::Copy(OutLengths.Slice(0, LhsLengths.Num()), LhsLengths);
+				RangesCopy(
+					OutStarts.Slice(0, LhsStarts.Num()),
+					OutLengths.Slice(0, LhsLengths.Num()),
+					LhsStarts,
+					LhsLengths);
+
 				return LhsLengths.Num();
 			}
 
@@ -1051,10 +1245,14 @@ namespace UE::Learning
 
 			// Return number of ranges added
 			return OutIndex;
+#endif
 		}
 
 		static inline int32 RangesTotalFrameNum(const TLearningArrayView<1, const int32> RangeLengths)
 		{
+#if UE_LEARNING_ISPC
+			return ispc::LearningRangesTotalFrameNum(RangeLengths.GetData(), RangeLengths.Num());
+#else
 			int32 Total = 0;
 
 			const int32 RangeNum = RangeLengths.Num();
@@ -1064,6 +1262,7 @@ namespace UE::Learning
 			}
 
 			return Total;
+#endif
 		}
 
 		static inline int32 RangesIntersectionWithOffsets(
@@ -1080,6 +1279,21 @@ namespace UE::Learning
 		{
 			TRACE_CPUPROFILER_EVENT_SCOPE(Learning::FrameRangeSet::Private::RangesIntersectionWithOffsets);
 
+#if UE_LEARNING_ISPC
+			return ispc::LearningRangesIntersectionWithOffsets(
+				OutStarts.GetData(),
+				OutLengths.GetData(),
+				OutLhsOffsets.GetData(),
+				OutRhsOffsets.GetData(),
+				LhsStarts.GetData(),
+				LhsLengths.GetData(),
+				RhsStarts.GetData(),
+				RhsLengths.GetData(),
+				LhsEntryOffset,
+				RhsEntryOffset,
+				LhsStarts.Num(),
+				RhsStarts.Num());
+#else
 			RangesCheck(LhsStarts, LhsLengths);
 			RangesCheck(RhsStarts, RhsLengths);
 
@@ -1088,8 +1302,12 @@ namespace UE::Learning
 
 			if (RangesEqual(LhsStarts, LhsLengths, RhsStarts, RhsLengths))
 			{
-				Array::Copy(OutStarts.Slice(0, LhsStarts.Num()), LhsStarts);
-				Array::Copy(OutLengths.Slice(0, LhsStarts.Num()), LhsLengths);
+				RangesCopy(
+					OutStarts.Slice(0, LhsStarts.Num()),
+					OutLengths.Slice(0, LhsLengths.Num()),
+					LhsStarts,
+					LhsLengths);
+
 				ComputeRangeOffsets(OutLhsOffsets.Slice(0, LhsStarts.Num()), LhsLengths, LhsEntryOffset);
 				ComputeRangeOffsets(OutRhsOffsets.Slice(0, RhsStarts.Num()), RhsLengths, RhsEntryOffset);
 				OffsetsCheck(OutLhsOffsets.Slice(0, LhsStarts.Num()), LhsEntryOffset, LhsEntryOffset + RangesTotalFrameNum(LhsLengths));
@@ -1240,6 +1458,7 @@ namespace UE::Learning
 
 			// Return number of ranges added
 			return OutIndex;
+#endif
 		}
 
 		static inline int32 RangesDifference(
@@ -1252,6 +1471,17 @@ namespace UE::Learning
 		{
 			TRACE_CPUPROFILER_EVENT_SCOPE(Learning::FrameRangeSet::Private::RangesDifference);
 
+#if UE_LEARNING_ISPC
+			return ispc::LearningRangesDifference(
+				OutStarts.GetData(),
+				OutLengths.GetData(),
+				LhsStarts.GetData(),
+				LhsLengths.GetData(),
+				RhsStarts.GetData(),
+				RhsLengths.GetData(),
+				LhsStarts.Num(),
+				RhsStarts.Num());
+#else
 			RangesCheck(LhsStarts, LhsLengths);
 			RangesCheck(RhsStarts, RhsLengths);
 
@@ -1259,8 +1489,12 @@ namespace UE::Learning
 
 			if (RhsStarts.IsEmpty())
 			{
-				Array::Copy(OutStarts.Slice(0, LhsStarts.Num()), LhsStarts);
-				Array::Copy(OutLengths.Slice(0, LhsLengths.Num()), LhsLengths);
+				RangesCopy(
+					OutStarts.Slice(0, LhsStarts.Num()),
+					OutLengths.Slice(0, LhsLengths.Num()),
+					LhsStarts,
+					LhsLengths);
+
 				return LhsLengths.Num();
 			}
 
@@ -1387,6 +1621,7 @@ namespace UE::Learning
 
 			// Return number of ranges added
 			return OutIndex;
+#endif
 		}
 	}
 
@@ -1664,8 +1899,11 @@ namespace UE::Learning
 		RangeStarts.SetNumUninitialized({ CurrRangeNum + AddRangeNum });
 		RangeLengths.SetNumUninitialized({ CurrRangeNum + AddRangeNum });
 		RangeOffsets.SetNumUninitialized({ CurrRangeNum + AddRangeNum });
-		Array::Copy(RangeStarts.Slice(CurrRangeNum, AddRangeNum), InStarts);
-		Array::Copy(RangeLengths.Slice(CurrRangeNum, AddRangeNum), InLengths);
+		FrameRangeSet::Private::RangesCopy(
+			RangeStarts.Slice(CurrRangeNum, AddRangeNum),
+			RangeLengths.Slice(CurrRangeNum, AddRangeNum),
+			InStarts,
+			InLengths);
 		for (int32 Idx = CurrRangeNum; Idx < CurrRangeNum + AddRangeNum; Idx++)
 		{
 			RangeOffsets[Idx] = Idx == 0 ? 0 : RangeOffsets[Idx - 1] + RangeLengths[Idx - 1];
@@ -1752,13 +1990,10 @@ namespace UE::Learning
 					OutFrameRangeSet.EntryRangeOffsets[OutIndex] = RangeIndex;
 					OutFrameRangeSet.EntryRangeNums[OutIndex] = FrameSet.GetEntryFrameNum(LhsIndex);
 
-					Array::Copy(
+					FrameRangeSet::Private::RangesCopyFrames(
 						OutFrameRangeSet.RangeStarts.Slice(RangeIndex, FrameSet.GetEntryFrameNum(LhsIndex)),
-						FrameSet.GetEntryFrames(LhsIndex));
-
-					Array::Set(
 						OutFrameRangeSet.RangeLengths.Slice(RangeIndex, FrameSet.GetEntryFrameNum(LhsIndex)),
-						1);
+						FrameSet.GetEntryFrames(LhsIndex));
 
 					RangeIndex += FrameSet.GetEntryFrameNum(LhsIndex);
 					OutIndex++;
@@ -1774,12 +2009,10 @@ namespace UE::Learning
 					OutFrameRangeSet.EntryRangeOffsets[OutIndex] = RangeIndex;
 					OutFrameRangeSet.EntryRangeNums[OutIndex] = FrameRangeSet.GetEntryRangeNum(RhsIndex);
 
-					Array::Copy(
+					FrameRangeSet::Private::RangesCopy(
 						OutFrameRangeSet.RangeStarts.Slice(RangeIndex, FrameRangeSet.GetEntryRangeNum(RhsIndex)),
-						FrameRangeSet.GetEntryRangeStarts(RhsIndex));
-
-					Array::Copy(
 						OutFrameRangeSet.RangeLengths.Slice(RangeIndex, FrameRangeSet.GetEntryRangeNum(RhsIndex)),
+						FrameRangeSet.GetEntryRangeStarts(RhsIndex),
 						FrameRangeSet.GetEntryRangeLengths(RhsIndex));
 
 					RangeIndex += FrameRangeSet.GetEntryRangeNum(RhsIndex);
@@ -1824,12 +2057,11 @@ namespace UE::Learning
 				OutFrameRangeSet.EntryRangeOffsets[OutIndex] = RangeIndex;
 				OutFrameRangeSet.EntryRangeNums[OutIndex] = FrameSet.GetEntryFrameNum(LhsIndex);
 
-				Array::Copy(
+				FrameRangeSet::Private::RangesCopyFrames(
 					OutFrameRangeSet.RangeStarts.Slice(RangeIndex, FrameSet.GetEntryFrameNum(LhsIndex)),
+					OutFrameRangeSet.RangeLengths.Slice(RangeIndex, FrameSet.GetEntryFrameNum(LhsIndex)),
 					FrameSet.GetEntryFrames(LhsIndex));
 
-				Array::Set(OutFrameRangeSet.RangeLengths.Slice(RangeIndex, FrameSet.GetEntryFrameNum(LhsIndex)), 1);
-
 				RangeIndex += FrameSet.GetEntryFrameNum(LhsIndex);
 				OutIndex++;
 				LhsIndex++;
@@ -1846,12 +2078,10 @@ namespace UE::Learning
 				OutFrameRangeSet.EntryRangeOffsets[OutIndex] = RangeIndex;
 				OutFrameRangeSet.EntryRangeNums[OutIndex] = FrameRangeSet.GetEntryRangeNum(RhsIndex);
 
-				Array::Copy(
+				FrameRangeSet::Private::RangesCopy(
 					OutFrameRangeSet.RangeStarts.Slice(RangeIndex, FrameRangeSet.GetEntryRangeNum(RhsIndex)),
-					FrameRangeSet.GetEntryRangeStarts(RhsIndex));
-
-				Array::Copy(
 					OutFrameRangeSet.RangeLengths.Slice(RangeIndex, FrameRangeSet.GetEntryRangeNum(RhsIndex)),
+					FrameRangeSet.GetEntryRangeStarts(RhsIndex),
 					FrameRangeSet.GetEntryRangeLengths(RhsIndex));
 
 				RangeIndex += FrameRangeSet.GetEntryRangeNum(RhsIndex);
@@ -1985,7 +2215,7 @@ namespace UE::Learning
 						OutFrameSet.EntryFrameOffsets[OutIndex] = FrameIndex;
 						OutFrameSet.EntryFrameNums[OutIndex] = FrameSet.GetEntryFrameNum(LhsIndex);
 
-						Array::Copy(
+						FrameRangeSet::Private::FramesCopy(
 							OutFrameSet.Frames.Slice(FrameIndex, FrameSet.GetEntryFrameNum(LhsIndex)),
 							FrameSet.GetEntryFrames(LhsIndex));
 
@@ -2037,7 +2267,7 @@ namespace UE::Learning
 					OutFrameSet.EntryFrameOffsets[OutIndex] = FrameIndex;
 					OutFrameSet.EntryFrameNums[OutIndex] = FrameSet.GetEntryFrameNum(LhsIndex);
 
-					Array::Copy(
+					FrameRangeSet::Private::FramesCopy(
 						OutFrameSet.Frames.Slice(FrameIndex, FrameSet.GetEntryFrameNum(LhsIndex)),
 						FrameSet.GetEntryFrames(LhsIndex));
 
@@ -2098,12 +2328,10 @@ namespace UE::Learning
 					OutFrameRangeSet.EntryRangeOffsets[OutIndex] = RangeIndex;
 					OutFrameRangeSet.EntryRangeNums[OutIndex] = FrameRangeSet.GetEntryRangeNum(LhsIndex);
 
-					Array::Copy(
+					FrameRangeSet::Private::RangesCopy(
 						OutFrameRangeSet.RangeStarts.Slice(RangeIndex, FrameRangeSet.GetEntryRangeNum(LhsIndex)),
-						FrameRangeSet.GetEntryRangeStarts(LhsIndex));
-
-					Array::Copy(
 						OutFrameRangeSet.RangeLengths.Slice(RangeIndex, FrameRangeSet.GetEntryRangeNum(LhsIndex)),
+						FrameRangeSet.GetEntryRangeStarts(LhsIndex),
 						FrameRangeSet.GetEntryRangeLengths(LhsIndex));
 
 					RangeIndex += FrameRangeSet.GetEntryRangeNum(LhsIndex);
@@ -2154,12 +2382,10 @@ namespace UE::Learning
 				OutFrameRangeSet.EntryRangeOffsets[OutIndex] = RangeIndex;
 				OutFrameRangeSet.EntryRangeNums[OutIndex] = FrameRangeSet.GetEntryRangeNum(LhsIndex);
 
-				Array::Copy(
+				FrameRangeSet::Private::RangesCopy(
 					OutFrameRangeSet.RangeStarts.Slice(RangeIndex, FrameRangeSet.GetEntryRangeNum(LhsIndex)),
-					FrameRangeSet.GetEntryRangeStarts(LhsIndex));
-
-				Array::Copy(
 					OutFrameRangeSet.RangeLengths.Slice(RangeIndex, FrameRangeSet.GetEntryRangeNum(LhsIndex)),
+					FrameRangeSet.GetEntryRangeStarts(LhsIndex),
 					FrameRangeSet.GetEntryRangeLengths(LhsIndex));
 
 				RangeIndex += FrameRangeSet.GetEntryRangeNum(LhsIndex);
@@ -2223,12 +2449,10 @@ namespace UE::Learning
 					Out.EntryRangeOffsets[OutIndex] = RangeIndex;
 					Out.EntryRangeNums[OutIndex] = Lhs.GetEntryRangeNum(LhsIndex);
 
-					Array::Copy(
+					FrameRangeSet::Private::RangesCopy(
 						Out.RangeStarts.Slice(RangeIndex, Lhs.GetEntryRangeNum(LhsIndex)),
-						Lhs.GetEntryRangeStarts(LhsIndex));
-
-					Array::Copy(
 						Out.RangeLengths.Slice(RangeIndex, Lhs.GetEntryRangeNum(LhsIndex)),
+						Lhs.GetEntryRangeStarts(LhsIndex),
 						Lhs.GetEntryRangeLengths(LhsIndex));
 
 					RangeIndex += Lhs.GetEntryRangeNum(LhsIndex);
@@ -2245,12 +2469,10 @@ namespace UE::Learning
 					Out.EntryRangeOffsets[OutIndex] = RangeIndex;
 					Out.EntryRangeNums[OutIndex] = Rhs.GetEntryRangeNum(RhsIndex);
 
-					Array::Copy(
+					FrameRangeSet::Private::RangesCopy(
 						Out.RangeStarts.Slice(RangeIndex, Rhs.GetEntryRangeNum(RhsIndex)),
-						Rhs.GetEntryRangeStarts(RhsIndex));
-
-					Array::Copy(
 						Out.RangeLengths.Slice(RangeIndex, Rhs.GetEntryRangeNum(RhsIndex)),
+						Rhs.GetEntryRangeStarts(RhsIndex),
 						Rhs.GetEntryRangeLengths(RhsIndex));
 
 					RangeIndex += Rhs.GetEntryRangeNum(RhsIndex);
@@ -2296,12 +2518,10 @@ namespace UE::Learning
 				Out.EntryRangeOffsets[OutIndex] = RangeIndex;
 				Out.EntryRangeNums[OutIndex] = Lhs.GetEntryRangeNum(LhsIndex);
 
-				Array::Copy(
+				FrameRangeSet::Private::RangesCopy(
 					Out.RangeStarts.Slice(RangeIndex, Lhs.GetEntryRangeNum(LhsIndex)),
-					Lhs.GetEntryRangeStarts(LhsIndex));
-
-				Array::Copy(
 					Out.RangeLengths.Slice(RangeIndex, Lhs.GetEntryRangeNum(LhsIndex)),
+					Lhs.GetEntryRangeStarts(LhsIndex),
 					Lhs.GetEntryRangeLengths(LhsIndex));
 
 				RangeIndex += Lhs.GetEntryRangeNum(LhsIndex);
@@ -2320,12 +2540,10 @@ namespace UE::Learning
 				Out.EntryRangeOffsets[OutIndex] = RangeIndex;
 				Out.EntryRangeNums[OutIndex] = Rhs.GetEntryRangeNum(RhsIndex);
 
-				Array::Copy(
+				FrameRangeSet::Private::RangesCopy(
 					Out.RangeStarts.Slice(RangeIndex, Rhs.GetEntryRangeNum(RhsIndex)),
-					Rhs.GetEntryRangeStarts(RhsIndex));
-
-				Array::Copy(
 					Out.RangeLengths.Slice(RangeIndex, Rhs.GetEntryRangeNum(RhsIndex)),
+					Rhs.GetEntryRangeStarts(RhsIndex),
 					Rhs.GetEntryRangeLengths(RhsIndex));
 
 				RangeIndex += Rhs.GetEntryRangeNum(RhsIndex);
@@ -2471,12 +2689,10 @@ namespace UE::Learning
 					Out.EntryRangeOffsets[OutIndex] = RangeIndex;
 					Out.EntryRangeNums[OutIndex] = Lhs.GetEntryRangeNum(LhsIndex);
 
-					Array::Copy(
+					FrameRangeSet::Private::RangesCopy(
 						Out.RangeStarts.Slice(RangeIndex, Lhs.GetEntryRangeNum(LhsIndex)),
-						Lhs.GetEntryRangeStarts(LhsIndex));
-
-					Array::Copy(
 						Out.RangeLengths.Slice(RangeIndex, Lhs.GetEntryRangeNum(LhsIndex)),
+						Lhs.GetEntryRangeStarts(LhsIndex),
 						Lhs.GetEntryRangeLengths(LhsIndex));
 
 					RangeIndex += Lhs.GetEntryRangeNum(LhsIndex);
@@ -2528,12 +2744,10 @@ namespace UE::Learning
 				Out.EntryRangeOffsets[OutIndex] = RangeIndex;
 				Out.EntryRangeNums[OutIndex] = Lhs.GetEntryRangeNum(LhsIndex);
 
-				Array::Copy(
+				FrameRangeSet::Private::RangesCopy(
 					Out.RangeStarts.Slice(RangeIndex, Lhs.GetEntryRangeNum(LhsIndex)),
-					Lhs.GetEntryRangeStarts(LhsIndex));
-
-				Array::Copy(
 					Out.RangeLengths.Slice(RangeIndex, Lhs.GetEntryRangeNum(LhsIndex)),
+					Lhs.GetEntryRangeStarts(LhsIndex),
 					Lhs.GetEntryRangeLengths(LhsIndex));
 
 				RangeIndex += Lhs.GetEntryRangeNum(LhsIndex);
diff --git a/Engine/Plugins/Experimental/LearningCore/Source/Learning/Public/LearningFrameAttribute.h b/Engine/Plugins/Experimental/LearningCore/Source/Learning/Public/LearningFrameAttribute.h
--- a/Engine/Plugins/Experimental/LearningCore/Source/Learning/Public/LearningFrameAttribute.h
+++ b/Engine/Plugins/Experimental/LearningCore/Source/Learning/Public/LearningFrameAttribute.h
@@ -559,6 +559,18 @@ namespace UE::Learning
 
 		/** Applies the transform to a vector without translation or scale (i.e. rotation only) */
 		UE_API void TransformConstantApplyNoTranslationScale(FFrameAttribute& Out, const FTransform3f& Lhs, const FFrameAttribute& Rhs);
+
+
+		/** Inertialization distance for a cubic inertialization function. See: https://theorangeduck.com/page/inertialization-transition-cost */
+		UE_API void InertializationDistanceConstant(
+			FFrameAttribute& Out, 
+			const FFrameAttribute& LhsPos, 
+			const FFrameAttribute& LhsVel,
+			const TArrayView<const float> RhsLoc,
+			const TArrayView<const float> RhsVel,
+			const float BlendTime,
+			const float LocWeight,
+			const float VelWeight);
 	}
 }
 
