LearningCore: Added some functions to compute average range length and total distance traveled.

#Synced-CL 50443644

diff --git a/Engine/Plugins/Experimental/LearningCore/Source/Learning/Private/LearningFrameAttribute.cpp b/Engine/Plugins/Experimental/LearningCore/Source/Learning/Private/LearningFrameAttribute.cpp
--- a/Engine/Plugins/Experimental/LearningCore/Source/Learning/Private/LearningFrameAttribute.cpp
+++ b/Engine/Plugins/Experimental/LearningCore/Source/Learning/Private/LearningFrameAttribute.cpp
@@ -1102,6 +1102,68 @@ namespace UE::Learning
 			OutLogStd = FMath::Sqrt(OutLogStd);
 		}
 		
+		static inline void MinMax(float& OutMin, float& OutMax, const TLearningArrayView<1, const float> In)
+		{
+			const int32 Num = In.Num();
+
+			OutMin = +UE_MAX_FLT;
+			OutMax = -UE_MAX_FLT;
+			for (int32 Idx = 0; Idx < Num; Idx++)
+			{
+				OutMin = FMath::Min(OutMin, In[Idx]);
+				OutMax = FMath::Max(OutMax, In[Idx]);
+			}
+		}
+
+		static inline void LocationDistanceTraveled(
+			float& OutDistance, 
+			const TLearningArrayView<1, const float> InX,
+			const TLearningArrayView<1, const float> InY, 
+			const TLearningArrayView<1, const float> InZ)
+		{
+			check(InX.Num() == InY.Num());
+			check(InX.Num() == InZ.Num());
+
+			const int32 Num = InX.Num();
+
+			OutDistance = 0.0f;
+			for (int32 Idx = 0; Idx < Num - 1; Idx++)
+			{
+				OutDistance += FMath::Sqrt(
+					FMath::Square(InX[Idx + 1] - InX[Idx + 0]) +
+					FMath::Square(InY[Idx + 1] - InY[Idx + 0]) +
+					FMath::Square(InZ[Idx + 1] - InZ[Idx + 0]));
+			}
+		}
+
+		static inline void QuatAngleTraveled(
+			float& OutAngle,
+			const TLearningArrayView<1, const float> InX,
+			const TLearningArrayView<1, const float> InY,
+			const TLearningArrayView<1, const float> InZ,
+			const TLearningArrayView<1, const float> InW)
+		{
+			check(InX.Num() == InY.Num());
+			check(InX.Num() == InZ.Num());
+			check(InX.Num() == InW.Num());
+
+			const int32 Num = InX.Num();
+
+			OutAngle = 0.0f;
+			for (int32 Idx = 0; Idx < Num - 1; Idx++)
+			{
+				OutAngle += FQuat4f(
+					InX[Idx + 1],
+					InY[Idx + 1],
+					InZ[Idx + 1],
+					InW[Idx + 1]).AngularDistance(FQuat4f(
+						InX[Idx + 0],
+						InY[Idx + 0],
+						InZ[Idx + 0],
+						InW[Idx + 0]));
+			}
+		}
+
 		static inline void QuatInv(TLearningArrayView<2, float> Out, const TLearningArrayView<2, const float> In)
 		{
 			check(Out.Num<0>() == 4);
@@ -3909,6 +3971,8 @@ namespace UE::Learning
 
 		void Mean(TLearningArrayView<1, float> OutMean, const FFrameAttribute& In)
 		{
+			check(OutMean.Num() == In.GetChannelNum());
+
 			ReduceOp(In, [&OutMean](
 				const FFrameAttribute& In,
 				const TLearningArrayView<1, const int32> RangeOffsets,
@@ -3925,6 +3989,9 @@ namespace UE::Learning
 
 		void MeanStd(TLearningArrayView<1, float> OutMean, TLearningArrayView<1, float> OutStd, const FFrameAttribute& In)
 		{
+			check(OutMean.Num() == In.GetChannelNum());
+			check(OutStd.Num() == In.GetChannelNum());
+
 			ReduceOp(In, [&OutMean, &OutStd](
 				const FFrameAttribute& In,
 				const TLearningArrayView<1, const int32> RangeOffsets,
@@ -3941,6 +4008,9 @@ namespace UE::Learning
 
 		void LogMeanStd(TLearningArrayView<1, float> OutMean, TLearningArrayView<1, float> OutLogStd, const FFrameAttribute& In)
 		{
+			check(OutMean.Num() == In.GetChannelNum());
+			check(OutLogStd.Num() == In.GetChannelNum());
+
 			ReduceOp(In, [&OutMean, &OutLogStd](
 				const FFrameAttribute& In,
 				const TLearningArrayView<1, const int32> RangeOffsets,
@@ -3955,6 +4025,25 @@ namespace UE::Learning
 				});
 		}
 
+		void MinMax(TLearningArrayView<1, float> OutMin, TLearningArrayView<1, float> OutMax, const FFrameAttribute& In)
+		{
+			check(OutMin.Num() == In.GetChannelNum());
+			check(OutMax.Num() == In.GetChannelNum());
+
+			ReduceOp(In, [&OutMin, &OutMax](
+				const FFrameAttribute& In,
+				const TLearningArrayView<1, const int32> RangeOffsets,
+				const TLearningArrayView<1, const int32> RangeLengths) {
+
+					const int32 ChannelNum = In.GetChannelNum();
+
+					for (int32 ChannelIdx = 0; ChannelIdx < ChannelNum; ChannelIdx++)
+					{
+						Private::MinMax(OutMin[ChannelIdx], OutMax[ChannelIdx], In.GetChannelAttributeData(ChannelIdx));
+					}
+				});
+		}
+
 		void QuatMul(FFrameAttribute& Out, const FFrameAttribute& Lhs, const FFrameAttribute& Rhs)
 		{
 			check(Lhs.GetChannelNum() == 4);
@@ -4880,5 +4969,53 @@ namespace UE::Learning
 					}
 				});
 		}
+
+		void LocationDistanceTraveled(TArrayView<float> OutDistancesTraveled, const FFrameAttribute& In)
+		{
+			check(In.GetChannelNum() == 3);
+			check(OutDistancesTraveled.Num() == In.FrameRangeSet.GetTotalRangeNum());
+
+			ReduceOp(In, [&OutDistancesTraveled](
+				const FFrameAttribute& In,
+				const TLearningArrayView<1, const int32> RangeOffsets,
+				const TLearningArrayView<1, const int32> RangeLengths) {
+
+					const int32 RangeNum = RangeOffsets.Num();
+
+					for (int32 RangeIdx = 0; RangeIdx < RangeNum; RangeIdx++)
+					{
+						Private::LocationDistanceTraveled(
+							OutDistancesTraveled[RangeIdx],
+							In.GetChannelRangeAttributeData(0, RangeOffsets[RangeIdx], RangeLengths[RangeIdx]),
+							In.GetChannelRangeAttributeData(1, RangeOffsets[RangeIdx], RangeLengths[RangeIdx]),
+							In.GetChannelRangeAttributeData(2, RangeOffsets[RangeIdx], RangeLengths[RangeIdx]));
+					}
+				});
+		}
+
+		void QuatAngleTraveled(TArrayView<float> OutAnglesTraveled, const FFrameAttribute& In)
+		{
+			check(In.GetChannelNum() == 4);
+			check(OutAnglesTraveled.Num() == In.FrameRangeSet.GetTotalRangeNum());
+
+			ReduceOp(In, [&OutAnglesTraveled](
+				const FFrameAttribute& In,
+				const TLearningArrayView<1, const int32> RangeOffsets,
+				const TLearningArrayView<1, const int32> RangeLengths) {
+
+					const int32 RangeNum = RangeOffsets.Num();
+
+					for (int32 RangeIdx = 0; RangeIdx < RangeNum; RangeIdx++)
+					{
+						Private::QuatAngleTraveled(
+							OutAnglesTraveled[RangeIdx],
+							In.GetChannelRangeAttributeData(0, RangeOffsets[RangeIdx], RangeLengths[RangeIdx]),
+							In.GetChannelRangeAttributeData(1, RangeOffsets[RangeIdx], RangeLengths[RangeIdx]),
+							In.GetChannelRangeAttributeData(2, RangeOffsets[RangeIdx], RangeLengths[RangeIdx]),
+							In.GetChannelRangeAttributeData(3, RangeOffsets[RangeIdx], RangeLengths[RangeIdx]));
+					}
+				});
+		}
+
 	}
 }
diff --git a/Engine/Plugins/Experimental/LearningCore/Source/Learning/Private/LearningFrameRangeSet.cpp b/Engine/Plugins/Experimental/LearningCore/Source/Learning/Private/LearningFrameRangeSet.cpp
--- a/Engine/Plugins/Experimental/LearningCore/Source/Learning/Private/LearningFrameRangeSet.cpp
+++ b/Engine/Plugins/Experimental/LearningCore/Source/Learning/Private/LearningFrameRangeSet.cpp
@@ -1707,6 +1707,35 @@ namespace UE::Learning
 		return RangeNum == 0 ? 0 : RangeOffsets[RangeNum - 1] + RangeLengths[RangeNum - 1];
 	}
 
+	float FFrameRangeSet::GetAverageFrameNum() const
+	{
+		return GetTotalRangeNum() == 0 ? 0 : (float)GetTotalFrameNum() / (float)GetTotalRangeNum();
+	}
+
+	int32 FFrameRangeSet::GetMinFrameNum() const
+	{
+		const int32 RangeNum = GetTotalRangeNum();
+		if (RangeNum == 0) { return 0; }
+		int32 Minimum = INT32_MAX;
+		for (int32 RangeIdx = 0; RangeIdx < RangeNum; RangeIdx++)
+		{
+			Minimum = FMath::Min(Minimum, RangeLengths[RangeIdx]);
+		}
+		return Minimum;
+	}
+
+	int32 FFrameRangeSet::GetMaxFrameNum() const
+	{
+		const int32 RangeNum = GetTotalRangeNum();
+		if (RangeNum == 0) { return 0; }
+		int32 Maximum = INT32_MIN;
+		for (int32 RangeIdx = 0; RangeIdx < RangeNum; RangeIdx++)
+		{
+			Maximum = FMath::Max(Maximum, RangeLengths[RangeIdx]);
+		}
+		return Maximum;
+	}
+
 	bool FFrameRangeSet::ContainsSequence(const int32 Sequence) const
 	{
 		return EntrySequences.ArrayView().Contains(Sequence);
diff --git a/Engine/Plugins/Experimental/LearningCore/Source/Learning/Public/LearningFrameAttribute.h b/Engine/Plugins/Experimental/LearningCore/Source/Learning/Public/LearningFrameAttribute.h
--- a/Engine/Plugins/Experimental/LearningCore/Source/Learning/Public/LearningFrameAttribute.h
+++ b/Engine/Plugins/Experimental/LearningCore/Source/Learning/Public/LearningFrameAttribute.h
@@ -432,6 +432,9 @@ namespace UE::Learning
 		/** Computes the mean and log std across all frames in the log space. OutMean and OutLogStd should be size of the number of channels */
 		UE_API void LogMeanStd(TLearningArrayView<1, float> OutMean, TLearningArrayView<1, float> OutLogStd, const FFrameAttribute& In);
 
+		/** Computes the min and max across all frames. OutMin and OutMax should be size of the number of channels */
+		UE_API void MinMax(TLearningArrayView<1, float> OutMin, TLearningArrayView<1, float> OutMax, const FFrameAttribute& In);
+
 
 		/** Computes the quaternion multiplication of two frame attributes. Both arguments are assumed to have 4 channels. */
 		UE_API void QuatMul(FFrameAttribute& Out, const FFrameAttribute& Lhs, const FFrameAttribute& Rhs);
@@ -571,6 +574,12 @@ namespace UE::Learning
 			const float BlendTime,
 			const float LocWeight,
 			const float VelWeight);
+
+		/** Gets distance traveled of each range in a location frame attribute */
+		UE_API void LocationDistanceTraveled(TArrayView<float> OutDistancesTraveled, const FFrameAttribute& In);
+
+		/** Gets angle traveled of each range in a quat frame attribute */
+		UE_API void QuatAngleTraveled(TArrayView<float> OutAnglesTraveled, const FFrameAttribute& In);
 	}
 }
 
diff --git a/Engine/Plugins/Experimental/LearningCore/Source/Learning/Public/LearningFrameRangeSet.h b/Engine/Plugins/Experimental/LearningCore/Source/Learning/Public/LearningFrameRangeSet.h
--- a/Engine/Plugins/Experimental/LearningCore/Source/Learning/Public/LearningFrameRangeSet.h
+++ b/Engine/Plugins/Experimental/LearningCore/Source/Learning/Public/LearningFrameRangeSet.h
@@ -108,6 +108,15 @@ namespace UE::Learning
 		/** Computes the total number of frames across all entries and ranges */
 		UE_API int32 GetTotalFrameNum() const;
 
+		/** Computes the average number of frames across all entries and ranges */
+		UE_API float GetAverageFrameNum() const;
+
+		/** Gets the largest number of frames across all entries and ranges */
+		UE_API int32 GetMaxFrameNum() const;
+
+		/** Gets the smallest number of frames across all entries and ranges */
+		UE_API int32 GetMinFrameNum() const;
+
 		/** Checks if this FrameRangeSet contains a given sequence */
 		UE_API bool ContainsSequence(const int32 Sequence) const;
 
