LearningCore: Added SavGol filter

#Synced-CL 50709564

diff --git a/Engine/Plugins/Experimental/LearningCore/Source/Learning/Private/Learning.ispc b/Engine/Plugins/Experimental/LearningCore/Source/Learning/Private/Learning.ispc
--- a/Engine/Plugins/Experimental/LearningCore/Source/Learning/Private/Learning.ispc
+++ b/Engine/Plugins/Experimental/LearningCore/Source/Learning/Private/Learning.ispc
@@ -2052,27 +2052,25 @@ export uniform int32 LearningRangesDifference(
 	return OutIndex;
 }
 
-export void LearningFilterGaussian(
+export void LearningFilterWithKernel(
 	uniform float Out[], 
 	uniform const float In[], 
+	uniform const float Kernel[],
 	uniform const int ValueNum,
-	uniform const float StdInFrames)
+	uniform const int KernelNum)
 {
-	uniform const int StdRange = (uniform int)round(StdInFrames * 3.0f);
+	uniform const int KernelHalfWidth = KernelNum / 2;
 
 	foreach (ValueIdx = 0 ... ValueNum)
 	{
 		float Total = 0.0f;
 		Out[ValueIdx] = 0.0f;
 
-		const int RangeMin = max(ValueIdx - StdRange, 0);
-		const int RangeMax = min(ValueIdx + StdRange, ValueNum - 1);
-
-		for (int Offset = -StdRange; Offset <= +StdRange; Offset++)
+		for (int Offset = -KernelHalfWidth; Offset <= +KernelHalfWidth; Offset++)
 		{
 			if (ValueIdx + Offset >= 0 && ValueIdx + Offset < ValueNum)
 			{
-				const float Weight = LearningInvExpApprox(LearningSquare(Offset / StdInFrames));
+				const float Weight = Kernel[Offset + KernelHalfWidth];
 				Out[ValueIdx] += Weight * In[ValueIdx + Offset];
 				Total += Weight;
 			}
diff --git a/Engine/Plugins/Experimental/LearningCore/Source/Learning/Private/LearningFrameAttribute.cpp b/Engine/Plugins/Experimental/LearningCore/Source/Learning/Private/LearningFrameAttribute.cpp
--- a/Engine/Plugins/Experimental/LearningCore/Source/Learning/Private/LearningFrameAttribute.cpp
+++ b/Engine/Plugins/Experimental/LearningCore/Source/Learning/Private/LearningFrameAttribute.cpp
@@ -4,6 +4,8 @@
 
 #include "Math/Transform.h"
 
+#include "LearningEigen.h"
+
 #if UE_LEARNING_ISPC
 #include "Learning.ispc.generated.h"
 #endif
@@ -994,17 +996,17 @@ namespace UE::Learning
 			}
 		}
 
-		static inline void FilterGaussian(TLearningArrayView<1, float> Out, const TLearningArrayView<1, const float> In, const float StdInFrames)
+		static inline void FilterWithKernel(TLearningArrayView<1, float> Out, const TLearningArrayView<1, const float> In, const TLearningArrayView<1, const float> Kernel)
 		{
-			check(Out.Num() == In.Num());
-			check(StdInFrames >= 0.0f);
+			check(Kernel.Num() % 2 == 1); // Kernel must be odd in length
 
 #if UE_LEARNING_ISPC
 
-			ispc::LearningFilterGaussian(Out.GetData(), In.GetData(), In.Num(), StdInFrames);
+			ispc::LearningFilterWithKernel(Out.GetData(), In.GetData(), Kernel.GetData(), In.Num(), Kernel.Num());
 
 #else
-			const int32 StdRange = FMath::RoundToInt(StdInFrames * 3.0f);
+			const int32 KernelNum = Kernel.Num();
+			const int32 KernelHalfWidth = FMath::DivideAndRoundDown(KernelNum, 2);
 
 			const int32 ValueNum = Out.Num();
 			for (int32 ValueIdx = 0; ValueIdx < ValueNum; ValueIdx++)
@@ -1012,14 +1014,11 @@ namespace UE::Learning
 				float Total = 0.0f;
 				Out[ValueIdx] = 0.0f;
 
-				const int32 RangeMin = FMath::Max(ValueIdx - StdRange, 0);
-				const int32 RangeMax = FMath::Min(ValueIdx + StdRange, ValueNum - 1);
-
-				for (int32 Offset = -StdRange; Offset <= +StdRange; Offset++)
+				for (int32 Offset = -KernelHalfWidth; Offset <= +KernelHalfWidth; Offset++)
 				{
 					if (ValueIdx + Offset >= 0 && ValueIdx + Offset < ValueNum)
 					{
-						const float Weight = FMath::InvExpApprox(FMath::Square(Offset / StdInFrames));
+						const float Weight = Kernel[Offset + KernelHalfWidth];
 						Out[ValueIdx] += Weight * In[ValueIdx + Offset];
 						Total += Weight;
 					}
@@ -1030,10 +1029,22 @@ namespace UE::Learning
 #endif
 		}
 
+		static inline void GaussianKernel(TLearningArrayView<1, float> Kernel, const float StdInFrames)
+		{
+			check(Kernel.Num() % 2 == 1);
+
+			const int32 KernelSize = Kernel.Num();
+
+			for (int32 KernelIdx = 0; KernelIdx < KernelSize; KernelIdx++)
+			{
+				Kernel[KernelIdx] = FMath::Exp(-FMath::Square(((float)KernelIdx + 0.5f - ((float)KernelSize / 2.0f)) / FMath::Max(StdInFrames, UE_SMALL_NUMBER)));
+			}
+		}
+
 		static inline void FilterMajorityVote(TLearningArrayView<1, float> Out, const TLearningArrayView<1, const float> In, const int32 FilterWidthInFrames)
 		{
 			check(Out.Num() == In.Num());
-			check(FilterWidthInFrames >= 0.0f);
+			check(FilterWidthInFrames >= 0);
 
 #if UE_LEARNING_ISPC
 
@@ -1045,9 +1056,6 @@ namespace UE::Learning
 			{
 				int32 Total = 0;
 
-				const int32 RangeMin = FMath::Max(ValueIdx - FilterWidthInFrames, 0);
-				const int32 RangeMax = FMath::Min(ValueIdx + FilterWidthInFrames, ValueNum - 1);
-
 				for (int32 Offset = -FilterWidthInFrames; Offset <= +FilterWidthInFrames; Offset++)
 				{
 					if (ValueIdx + Offset >= 0 && ValueIdx + Offset < ValueNum)
@@ -1061,6 +1069,56 @@ namespace UE::Learning
 #endif
 		}
 
+		static inline void SavGolKernel(TLearningArrayView<1, float> Kernel, const int32 PolynomialDegree, bool bGaussianWindowed)
+		{
+			check(Kernel.Num() % 2 == 1);
+
+			const int32 KernelSize = Kernel.Num();
+			const int32 Order = PolynomialDegree + 1;
+
+			TLearningArray<2, float> Lhs;
+			Lhs.SetNumUninitialized({ KernelSize, Order });
+			for (int32 KernelIdx = 0; KernelIdx < KernelSize; KernelIdx++)
+			{
+				for (int32 OrderIdx = 0; OrderIdx < Order; OrderIdx++)
+				{
+					Lhs[KernelIdx][OrderIdx] = FMath::Pow((float)KernelIdx, (float)OrderIdx);
+				}
+			}
+
+			if (bGaussianWindowed)
+			{
+				for (int32 KernelIdx = 0; KernelIdx < KernelSize; KernelIdx++)
+				{
+					const float Weight = FMath::Exp(-FMath::Square(((float)KernelIdx + 0.5f - ((float)KernelSize / 2.0f)) / (1.5f * KernelSize)));
+
+					for (int32 OrderIdx = 0; OrderIdx < Order; OrderIdx++)
+					{
+						Lhs[KernelIdx][OrderIdx] *= Weight;
+					}
+				}
+			}
+
+			TLearningArray<2, float> Rhs;
+			Rhs.SetNumZeroed({ KernelSize, 1 });
+			Rhs[KernelSize / 2][0] = 1.0f;
+
+			TLearningArray<2, float> Coeff;
+			Coeff.SetNumZeroed({ Order, 1 });
+
+			OutEigenMatrix(Coeff).noalias() = InEigenMatrix(Lhs).bdcSvd(Eigen::ComputeThinU | Eigen::ComputeThinV).solve(InEigenMatrix(Rhs));
+
+			Array::Zero(Kernel);
+
+			for (int32 KernelIdx = 0; KernelIdx < KernelSize; KernelIdx++)
+			{
+				for (int32 OrderIdx = 0; OrderIdx < Order; OrderIdx++)
+				{
+					Kernel[KernelIdx] += Coeff[OrderIdx][0] * FMath::Pow((float)KernelIdx, (float)OrderIdx);
+				}
+			}
+		}
+
 		static inline void Mean(float& OutMean, const TLearningArrayView<1, const float> In)
 		{
 			const int32 Num = In.Num();
@@ -3925,7 +3983,14 @@ namespace UE::Learning
 
 		void FilterGaussian(FFrameAttribute& Out, const FFrameAttribute& In, const float StdInFrames)
 		{
-			UnaryOp(Out, In.GetChannelNum(), In, [StdInFrames](
+			check(StdInFrames >= 0.0f);
+
+			// Create odd sized kernel that is 3x the std
+			TArray<float, TInlineAllocator<128>> Kernel;
+			Kernel.SetNumUninitialized(FMath::RoundToInt(StdInFrames * 3.0f) * 2 + 1);
+			Private::GaussianKernel(Kernel, StdInFrames);
+
+			UnaryOp(Out, In.GetChannelNum(), In, [&Kernel](
 				FFrameAttribute& Out,
 				const FFrameAttribute& In,
 				const TLearningArrayView<1, const int32> RangeOffsets,
@@ -3937,10 +4002,10 @@ namespace UE::Learning
 					{
 						for (int32 RangeIdx = 0; RangeIdx < RangeNum; RangeIdx++)
 						{
-							Private::FilterGaussian(
+							Private::FilterWithKernel(
 								Out.GetChannelRangeAttributeData(ChannelIdx, RangeOffsets[RangeIdx], RangeLengths[RangeIdx]),
 								In.GetChannelRangeAttributeData(ChannelIdx, RangeOffsets[RangeIdx], RangeLengths[RangeIdx]),
-								StdInFrames);
+								Kernel);
 						}
 					}
 				});
@@ -3948,6 +4013,8 @@ namespace UE::Learning
 
 		void FilterMajorityVote(FFrameAttribute& Out, const FFrameAttribute& In, const int32 FilterWidthFrames)
 		{
+			check(FilterWidthFrames >= 0);
+
 			UnaryOp(Out, In.GetChannelNum(), In, [FilterWidthFrames](
 				FFrameAttribute& Out,
 				const FFrameAttribute& In,
@@ -3969,6 +4036,37 @@ namespace UE::Learning
 				});
 		}
 
+		void FilterSavGol(FFrameAttribute& Out, const FFrameAttribute& In, const int32 FilterWidthFrames, const int32 PolynomialDegree, const bool bGaussianWindowed)
+		{
+			check(FilterWidthFrames >= 0);
+			check(PolynomialDegree >= 0);
+
+			// Create odd sized kernel
+			TArray<float, TInlineAllocator<128>> Kernel;
+			Kernel.SetNumUninitialized((FilterWidthFrames / 2) * 2 + 1);
+			Private::SavGolKernel(Kernel, PolynomialDegree, bGaussianWindowed);
+
+			UnaryOp(Out, In.GetChannelNum(), In, [&Kernel](
+				FFrameAttribute& Out,
+				const FFrameAttribute& In,
+				const TLearningArrayView<1, const int32> RangeOffsets,
+				const TLearningArrayView<1, const int32> RangeLengths) {
+
+					const int32 ChannelNum = Out.GetChannelNum();
+					const int32 RangeNum = RangeLengths.Num();
+					for (int32 ChannelIdx = 0; ChannelIdx < ChannelNum; ChannelIdx++)
+					{
+						for (int32 RangeIdx = 0; RangeIdx < RangeNum; RangeIdx++)
+						{
+							Private::FilterWithKernel(
+								Out.GetChannelRangeAttributeData(ChannelIdx, RangeOffsets[RangeIdx], RangeLengths[RangeIdx]),
+								In.GetChannelRangeAttributeData(ChannelIdx, RangeOffsets[RangeIdx], RangeLengths[RangeIdx]),
+								Kernel);
+						}
+					}
+				});
+		}
+
 		void Mean(TLearningArrayView<1, float> OutMean, const FFrameAttribute& In)
 		{
 			check(OutMean.Num() == In.GetChannelNum());
diff --git a/Engine/Plugins/Experimental/LearningCore/Source/Learning/Public/LearningFrameAttribute.h b/Engine/Plugins/Experimental/LearningCore/Source/Learning/Public/LearningFrameAttribute.h
--- a/Engine/Plugins/Experimental/LearningCore/Source/Learning/Public/LearningFrameAttribute.h
+++ b/Engine/Plugins/Experimental/LearningCore/Source/Learning/Public/LearningFrameAttribute.h
@@ -423,6 +423,10 @@ namespace UE::Learning
 		/** Applies a Majority Vote filter to the ranges of values in the frame attribute. */
 		UE_API void FilterMajorityVote(FFrameAttribute& Out, const FFrameAttribute& In, const int32 FilterWidthFrames);
 
+		/** Applies a SavGol filter to the ranges of values in the frame attribute. */
+		UE_API void FilterSavGol(FFrameAttribute& Out, const FFrameAttribute& In, const int32 FilterWidthFrames, const int32 PolynomialDegree, const bool bGaussianWindowed);
+
+
 		/** Computes the mean across all frames. OutMean should be size of the number of channels */
 		UE_API void Mean(TLearningArrayView<1, float> OutMean, const FFrameAttribute& In);
 
diff --git a/Engine/Plugins/Experimental/LearningCore/Source/Learning/Public/MultiArray.h b/Engine/Plugins/Experimental/LearningCore/Source/Learning/Public/MultiArray.h
--- a/Engine/Plugins/Experimental/LearningCore/Source/Learning/Public/MultiArray.h
+++ b/Engine/Plugins/Experimental/LearningCore/Source/Learning/Public/MultiArray.h
@@ -702,7 +702,7 @@ public:
 		if (NewTotal > CurrTotal)
 		{
 			AllocatorResizeAllocation(CurrTotal, NewTotal);
-			FMemory::Memzero((void*)(GetData() + CurrTotal), (CurrTotal - NewTotal) * sizeof(ElementType));
+			FMemory::Memzero((void*)(GetData() + CurrTotal), (NewTotal - CurrTotal) * sizeof(ElementType));
 		}
 		else if (NewTotal < CurrTotal)
 		{
