LearningCore: Added a couple more frame attribute functions

#Synced-CL 50905131

diff --git a/Engine/Plugins/Experimental/LearningCore/Source/Learning/Private/LearningFrameAttribute.cpp b/Engine/Plugins/Experimental/LearningCore/Source/Learning/Private/LearningFrameAttribute.cpp
--- a/Engine/Plugins/Experimental/LearningCore/Source/Learning/Private/LearningFrameAttribute.cpp
+++ b/Engine/Plugins/Experimental/LearningCore/Source/Learning/Private/LearningFrameAttribute.cpp
@@ -1144,6 +1144,36 @@ namespace UE::Learning
 			OutStd = FMath::Sqrt(OutStd);
 		}
 
+		static inline void AngularMean(float& OutMean, const TLearningArrayView<1, const float> In)
+		{
+			const int32 Num = In.Num();
+
+			float MeanSin = 0.0f;
+			float MeanCos = 0.0f;
+
+			for (int32 Idx = 0; Idx < Num; Idx++)
+			{
+				MeanSin += (FMath::Sin(In[Idx]) - MeanSin) / (Idx + 1);
+				MeanCos += (FMath::Cos(In[Idx]) - MeanCos) / (Idx + 1);
+			}
+
+			OutMean = FMath::Atan2(MeanSin, MeanCos);
+		}
+
+		static inline void AngularMeanStd(float& OutMean, float& OutStd, const TLearningArrayView<1, const float> In)
+		{
+			AngularMean(OutMean, In);
+
+			const int32 Num = In.Num();
+
+			OutStd = 0.0f;
+			for (int32 Idx = 0; Idx < Num; Idx++)
+			{
+				OutStd += FMath::Square(FMath::FindDeltaAngleRadians(OutMean, In[Idx])) / Num;
+			}
+			OutStd = FMath::Sqrt(OutStd);
+		}
+
 		static inline void LogMeanStd(float& OutMean, float& OutLogStd, const TLearningArrayView<1, const float> In)
 		{
 			const int32 Num = In.Num();
@@ -1684,6 +1714,57 @@ namespace UE::Learning
 			}
 		}
 
+		static inline void QuatYaw(TLearningArrayView<2, float> Out, const TLearningArrayView<2, const float> In)
+		{
+			check(Out.Num<0>() == 1);
+			check(In.Num<0>() == 4);
+			check(Out.Num<1>() == In.Num<1>());
+
+			const int32 ValueNum = Out.Num<1>();
+			for (int32 ValueIdx = 0; ValueIdx < ValueNum; ValueIdx++)
+			{
+				const FQuat4f InValue = FQuat4f(In[0][ValueIdx], In[1][ValueIdx], In[2][ValueIdx], In[3][ValueIdx]);
+				Out[0][ValueIdx] = FMath::DegreesToRadians(InValue.Rotator().Yaw);
+			}
+		}
+
+		static inline void QuatPitch(TLearningArrayView<2, float> Out, const TLearningArrayView<2, const float> In)
+		{
+			check(Out.Num<0>() == 1);
+			check(In.Num<0>() == 4);
+			check(Out.Num<1>() == In.Num<1>());
+
+			const int32 ValueNum = Out.Num<1>();
+			for (int32 ValueIdx = 0; ValueIdx < ValueNum; ValueIdx++)
+			{
+				const FQuat4f InValue = FQuat4f(In[0][ValueIdx], In[1][ValueIdx], In[2][ValueIdx], In[3][ValueIdx]);
+				Out[0][ValueIdx] = FMath::DegreesToRadians(InValue.Rotator().Pitch);
+			}
+		}
+
+		static inline void QuatRoll(TLearningArrayView<2, float> Out, const TLearningArrayView<2, const float> In)
+		{
+			check(Out.Num<0>() == 1);
+			check(In.Num<0>() == 4);
+			check(Out.Num<1>() == In.Num<1>());
+
+			const int32 ValueNum = Out.Num<1>();
+			for (int32 ValueIdx = 0; ValueIdx < ValueNum; ValueIdx++)
+			{
+				const FQuat4f InValue = FQuat4f(In[0][ValueIdx], In[1][ValueIdx], In[2][ValueIdx], In[3][ValueIdx]);
+				Out[0][ValueIdx] = FMath::DegreesToRadians(InValue.Rotator().Roll);
+			}
+		}
+
+		static inline void WrapAngleInplace(TLearningArrayView<1, float> InOut)
+		{
+			const int32 ValueNum = InOut.Num();
+			for (int32 ValueIdx = 0; ValueIdx < ValueNum; ValueIdx++)
+			{
+				InOut[ValueIdx] = FMath::Wrap(InOut[ValueIdx], -UE_PI, UE_PI);
+			}
+		}
+
 		static inline void InertializationDistanceConstant(
 			TLearningArrayView<1, float> OutValues,
 			const TLearningArrayView<1, const float> Locations,
@@ -4104,6 +4185,43 @@ namespace UE::Learning
 				});
 		}
 
+		void AngularMean(TLearningArrayView<1, float> OutMean, const FFrameAttribute& In)
+		{
+			check(OutMean.Num() == In.GetChannelNum());
+
+			ReduceOp(In, [&OutMean](
+				const FFrameAttribute& In,
+				const TLearningArrayView<1, const int32> RangeOffsets,
+				const TLearningArrayView<1, const int32> RangeLengths) {
+
+					const int32 ChannelNum = In.GetChannelNum();
+
+					for (int32 ChannelIdx = 0; ChannelIdx < ChannelNum; ChannelIdx++)
+					{
+						Private::AngularMean(OutMean[ChannelIdx], In.GetChannelAttributeData(ChannelIdx));
+					}
+				});
+		}
+
+		void AngularMeanStd(TLearningArrayView<1, float> OutMean, TLearningArrayView<1, float> OutStd, const FFrameAttribute& In)
+		{
+			check(OutMean.Num() == In.GetChannelNum());
+			check(OutStd.Num() == In.GetChannelNum());
+
+			ReduceOp(In, [&OutMean, &OutStd](
+				const FFrameAttribute& In,
+				const TLearningArrayView<1, const int32> RangeOffsets,
+				const TLearningArrayView<1, const int32> RangeLengths) {
+
+					const int32 ChannelNum = In.GetChannelNum();
+
+					for (int32 ChannelIdx = 0; ChannelIdx < ChannelNum; ChannelIdx++)
+					{
+						Private::AngularMeanStd(OutMean[ChannelIdx], OutStd[ChannelIdx], In.GetChannelAttributeData(ChannelIdx));
+					}
+				});
+		}
+
 		void LogMeanStd(TLearningArrayView<1, float> OutMean, TLearningArrayView<1, float> OutLogStd, const FFrameAttribute& In)
 		{
 			check(OutMean.Num() == In.GetChannelNum());
@@ -4651,6 +4769,67 @@ namespace UE::Learning
 				});
 		}
 
+		void QuatYaw(FFrameAttribute& Out, const FFrameAttribute& In)
+		{
+			check(In.GetChannelNum() == 4);
+
+			UnaryOp(Out, 1, In, [](
+				FFrameAttribute& Out,
+				const FFrameAttribute& In,
+				const TLearningArrayView<1, const int32> RangeOffsets,
+				const TLearningArrayView<1, const int32> RangeLengths) {
+
+					Private::QuatYaw(Out.GetAttributeData(), In.GetAttributeData());
+				});
+		}
+
+		void QuatPitch(FFrameAttribute& Out, const FFrameAttribute& In)
+		{
+			check(In.GetChannelNum() == 4);
+
+			UnaryOp(Out, 1, In, [](
+				FFrameAttribute& Out,
+				const FFrameAttribute& In,
+				const TLearningArrayView<1, const int32> RangeOffsets,
+				const TLearningArrayView<1, const int32> RangeLengths) {
+
+					Private::QuatPitch(Out.GetAttributeData(), In.GetAttributeData());
+				});
+		}
+
+		void QuatRoll(FFrameAttribute& Out, const FFrameAttribute& In)
+		{
+			check(In.GetChannelNum() == 4);
+
+			UnaryOp(Out, 1, In, [](
+				FFrameAttribute& Out,
+				const FFrameAttribute& In,
+				const TLearningArrayView<1, const int32> RangeOffsets,
+				const TLearningArrayView<1, const int32> RangeLengths) {
+
+					Private::QuatRoll(Out.GetAttributeData(), In.GetAttributeData());
+				});
+		}
+
+
+		void WrapAngleInplace(FFrameAttribute& InOut)
+		{
+			InplaceOp(InOut, [](
+				FFrameAttribute& InOut,
+				const TLearningArrayView<1, const int32> RangeOffsets,
+				const TLearningArrayView<1, const int32> RangeLengths) {
+
+					const int32 ChannelNum = InOut.GetChannelNum();
+
+					for (int32 ChannelIdx = 0; ChannelIdx < ChannelNum; ChannelIdx++)
+					{
+						Private::WrapAngleInplace(
+							InOut.GetChannelAttributeData(ChannelIdx));
+					}
+				});
+		}
+
+
 		void TransformMake(FFrameAttribute& Out, const FFrameAttribute& Location, const FFrameAttribute& Rotation, const FFrameAttribute& Scale)
 		{
 			check(Location.GetChannelNum() == 3);
diff --git a/Engine/Plugins/Experimental/LearningCore/Source/Learning/Public/LearningFrameAttribute.h b/Engine/Plugins/Experimental/LearningCore/Source/Learning/Public/LearningFrameAttribute.h
--- a/Engine/Plugins/Experimental/LearningCore/Source/Learning/Public/LearningFrameAttribute.h
+++ b/Engine/Plugins/Experimental/LearningCore/Source/Learning/Public/LearningFrameAttribute.h
@@ -439,6 +439,12 @@ namespace UE::Learning
 		/** Computes the min and max across all frames. OutMin and OutMax should be size of the number of channels */
 		UE_API void MinMax(TLearningArrayView<1, float> OutMin, TLearningArrayView<1, float> OutMax, const FFrameAttribute& In);
 
+		/** Computes the angular mean across all frames. */
+		UE_API void AngularMean(TLearningArrayView<1, float> OutMean, const FFrameAttribute& In);
+
+		/** Computes the angular mean and std across all frames. */
+		UE_API void AngularMeanStd(TLearningArrayView<1, float> OutMean, TLearningArrayView<1, float> OutStd, const FFrameAttribute& In);
+
 
 		/** Computes the quaternion multiplication of two frame attributes. Both arguments are assumed to have 4 channels. */
 		UE_API void QuatMul(FFrameAttribute& Out, const FFrameAttribute& Lhs, const FFrameAttribute& Rhs);
@@ -521,6 +527,19 @@ namespace UE::Learning
 		/** Computes the quaternion mean and std across all frames. */
 		UE_API void QuatMeanStd(FQuat4f& OutMean, FVector3f& OutStd, const FFrameAttribute& In);
 
+		/** Computes the yaw component of a quaternion rotation. Argument is assumed to have 4 channels. */
+		UE_API void QuatYaw(FFrameAttribute& Out, const FFrameAttribute& In);
+
+		/** Computes the pitch component of a quaternion rotation. Argument is assumed to have 4 channels. */
+		UE_API void QuatPitch(FFrameAttribute& Out, const FFrameAttribute& In);
+
+		/** Computes the roll component of a quaternion rotation. Argument is assumed to have 4 channels. */
+		UE_API void QuatRoll(FFrameAttribute& Out, const FFrameAttribute& In);
+
+
+		/** Wraps a frame attribute in-place into the range -UE_PI, UE_PI */
+		UE_API void WrapAngleInplace(FFrameAttribute& InOut);
+
 
 		/** Make a transform attribute from a location, rotation, and scale */
 		UE_API void TransformMake(FFrameAttribute& Out, const FFrameAttribute& Location, const FFrameAttribute& Rotation, const FFrameAttribute& Scale);
