LearningCore: Added function to get Entry Stops

#Synced-CL 50905131

#Synced-CL 50905131

diff --git a/Engine/Plugins/Experimental/LearningCore/Source/Learning/Private/LearningFrameRangeSet.cpp b/Engine/Plugins/Experimental/LearningCore/Source/Learning/Private/LearningFrameRangeSet.cpp
--- a/Engine/Plugins/Experimental/LearningCore/Source/Learning/Private/LearningFrameRangeSet.cpp
+++ b/Engine/Plugins/Experimental/LearningCore/Source/Learning/Private/LearningFrameRangeSet.cpp
@@ -1678,6 +1678,7 @@ namespace UE::Learning
 	TLearningArrayView<1, const int32> FFrameRangeSet::GetEntryRangeOffsets(const int32 EntryIdx) const { return RangeOffsets.Slice(EntryRangeOffsets[EntryIdx], EntryRangeNums[EntryIdx]); }
 	int32 FFrameRangeSet::GetEntryRangeStart(const int32 EntryIdx, const int32 RangeIdx) const { return RangeStarts[EntryRangeOffsets[EntryIdx] + RangeIdx]; }
 	int32 FFrameRangeSet::GetEntryRangeLength(const int32 EntryIdx, const int32 RangeIdx) const { return RangeLengths[EntryRangeOffsets[EntryIdx] + RangeIdx]; }
+	int32 FFrameRangeSet::GetEntryRangeStop(const int32 EntryIdx, const int32 RangeIdx) const { return RangeStarts[EntryRangeOffsets[EntryIdx] + RangeIdx] + RangeLengths[EntryRangeOffsets[EntryIdx] + RangeIdx]; }
 	int32 FFrameRangeSet::GetEntryRangeOffset(const int32 EntryIdx, const int32 RangeIdx) const { return RangeOffsets[EntryRangeOffsets[EntryIdx] + RangeIdx]; }
 
 	float FFrameRangeSet::GetEntryRangeStartTime(const int32 EntryIdx, const int32 RangeIdx, const float FrameDeltaTime) const
diff --git a/Engine/Plugins/Experimental/LearningCore/Source/Learning/Public/LearningFrameRangeSet.h b/Engine/Plugins/Experimental/LearningCore/Source/Learning/Public/LearningFrameRangeSet.h
--- a/Engine/Plugins/Experimental/LearningCore/Source/Learning/Public/LearningFrameRangeSet.h
+++ b/Engine/Plugins/Experimental/LearningCore/Source/Learning/Public/LearningFrameRangeSet.h
@@ -81,6 +81,9 @@ namespace UE::Learning
 		/** Gets the range length associated with a given entry and range index */
 		UE_API int32 GetEntryRangeLength(const int32 EntryIdx, const int32 RangeIdx) const;
 
+		/** Gets the range stop associated with a given entry and range index */
+		UE_API int32 GetEntryRangeStop(const int32 EntryIdx, const int32 RangeIdx) const;
+
 		/** Computes the offset associated with a given entry and range index */
 		UE_API int32 GetEntryRangeOffset(const int32 EntryIdx, const int32 RangeIdx) const;
 
