LearningCore: Added a few more frame attribute functions

#Synced-CL 51496719

diff --git a/Engine/Plugins/Experimental/LearningCore/Source/Learning/Private/LearningFrameAttribute.cpp b/Engine/Plugins/Experimental/LearningCore/Source/Learning/Private/LearningFrameAttribute.cpp
--- a/Engine/Plugins/Experimental/LearningCore/Source/Learning/Private/LearningFrameAttribute.cpp
+++ b/Engine/Plugins/Experimental/LearningCore/Source/Learning/Private/LearningFrameAttribute.cpp
@@ -1323,6 +1323,54 @@ namespace UE::Learning
 			}
 		}
 
+		// See: https://theorangeduck.com/page/variations-muller
+		static inline void QuatFromBasisDirections(
+			TLearningArrayView<1, float> OutX,
+			TLearningArrayView<1, float> OutY,
+			TLearningArrayView<1, float> OutZ,
+			TLearningArrayView<1, float> OutW,
+			const TLearningArrayView<1, const float> InForwardX,
+			const TLearningArrayView<1, const float> InForwardY,
+			const TLearningArrayView<1, const float> InForwardZ,
+			const TLearningArrayView<1, const float> InRightX,
+			const TLearningArrayView<1, const float> InRightY,
+			const TLearningArrayView<1, const float> InRightZ,
+			const TLearningArrayView<1, const float> InUpX,
+			const TLearningArrayView<1, const float> InUpY,
+			const TLearningArrayView<1, const float> InUpZ,
+			const int32 Iterations = 30,
+			const float Eps = UE_SMALL_NUMBER)
+		{
+			const int32 ValueNum = OutX.Num();
+
+			for (int32 ValueIdx = 0; ValueIdx < ValueNum; ValueIdx++)
+			{
+				const FVector3f C0 = FVector3f(InForwardX[ValueIdx], InForwardY[ValueIdx], InForwardZ[ValueIdx]);
+				const FVector3f C1 = FVector3f(InRightX[ValueIdx], InRightY[ValueIdx], InRightZ[ValueIdx]);
+				const FVector3f C2 = FVector3f(InUpX[ValueIdx], InUpY[ValueIdx], InUpZ[ValueIdx]);
+				const float MLen = C0.Length() + C1.Length() + C2.Length() + Eps;
+
+				FQuat4f Q = FQuat4f::Identity;
+
+				for (int32 Iteration = 0; Iteration < Iterations; Iteration++)
+				{
+					const FVector3f R0 = Q.RotateVector(FVector3f::ForwardVector);
+					const FVector3f R1 = Q.RotateVector(FVector3f::RightVector);
+					const FVector3f R2 = Q.RotateVector(FVector3f::UpVector);
+
+					const float W = MLen + R0.Dot(C0) + R1.Dot(C1) + R2.Dot(C2);
+					const FVector3f V = R0.Cross(C0) + R1.Cross(C1) + R2.Cross(C2);
+
+					Q = FQuat4f(V.X, V.Y, V.Z, W).GetNormalized() * Q;
+				}
+
+				OutX[ValueIdx] = Q.X;
+				OutY[ValueIdx] = Q.Y;
+				OutZ[ValueIdx] = Q.Z;
+				OutW[ValueIdx] = Q.W;
+			}
+		}
+
 		static inline void QuatMul(
 			TLearningArrayView<1, float> OutX,
 			TLearningArrayView<1, float> OutY,
@@ -1928,6 +1976,54 @@ namespace UE::Learning
 			}
 		}
 
+		static inline void TransformApplyConstantLocation(
+			TLearningArrayView<2, float> Out,
+			const TLearningArrayView<2, const float> Lhs,
+			const FVector3f& Rhs)
+		{
+			check(Out.Num<0>() == 3);
+			check(Lhs.Num<0>() == 10);
+			check(Out.Num<1>() == Lhs.Num<1>());
+
+			const int32 ValueNum = Lhs.Num<1>();
+
+			for (int32 ValueIdx = 0; ValueIdx < ValueNum; ValueIdx++)
+			{
+				const FVector3f V = FTransform3f(
+					FQuat4f(Lhs[3][ValueIdx], Lhs[4][ValueIdx], Lhs[5][ValueIdx], Lhs[6][ValueIdx]),
+					FVector3f(Lhs[0][ValueIdx], Lhs[1][ValueIdx], Lhs[2][ValueIdx]),
+					FVector3f(Lhs[7][ValueIdx], Lhs[8][ValueIdx], Lhs[9][ValueIdx])).TransformPosition(Rhs);
+
+				Out[0][ValueIdx] = V.X;
+				Out[1][ValueIdx] = V.Y;
+				Out[2][ValueIdx] = V.Z;
+			}
+		}
+
+		static inline void TransformApplyConstantDirection(
+			TLearningArrayView<2, float> Out,
+			const TLearningArrayView<2, const float> Lhs,
+			const FVector3f& Rhs)
+		{
+			check(Out.Num<0>() == 3);
+			check(Lhs.Num<0>() == 10);
+			check(Out.Num<1>() == Lhs.Num<1>());
+
+			const int32 ValueNum = Lhs.Num<1>();
+
+			for (int32 ValueIdx = 0; ValueIdx < ValueNum; ValueIdx++)
+			{
+				const FVector3f V = FTransform3f(
+					FQuat4f(Lhs[3][ValueIdx], Lhs[4][ValueIdx], Lhs[5][ValueIdx], Lhs[6][ValueIdx]),
+					FVector3f(Lhs[0][ValueIdx], Lhs[1][ValueIdx], Lhs[2][ValueIdx]),
+					FVector3f(Lhs[7][ValueIdx], Lhs[8][ValueIdx], Lhs[9][ValueIdx])).TransformVectorNoScale(Rhs);
+
+				Out[0][ValueIdx] = V.X;
+				Out[1][ValueIdx] = V.Y;
+				Out[2][ValueIdx] = V.Z;
+			}
+		}
+
 		static inline void TransformConstantApply(
 			TLearningArrayView<2, float> Out,
 			const FTransform3f& Lhs,
@@ -4356,6 +4452,47 @@ namespace UE::Learning
 				});
 		}
 
+		void QuatFromBasisDirections(FFrameAttribute& Out, const FFrameAttribute& Forward, const FFrameAttribute& Right, const FFrameAttribute& Up)
+		{
+			check(Forward.GetChannelNum() == 3);
+			check(Right.GetChannelNum() == 3);
+			check(Up.GetChannelNum() == 3);
+
+			NaryOp(Out, 4, {
+				(ConstFrameAttributePtr)&Forward,
+				(ConstFrameAttributePtr)&Right,
+				(ConstFrameAttributePtr)&Up }, [](
+					FFrameAttribute& Out,
+					const TArrayView<const ConstFrameAttributePtr> In,
+					const TLearningArrayView<1, const int32> OutRangeOffsets,
+					const TArrayView<const TLearningArrayView<1, const int32>> InRangeOffsets,
+					const TLearningArrayView<1, const int32> RangeLengths) {
+
+						const int32 RangeNum = RangeLengths.Num();
+
+						for (int32 RangeIdx = 0; RangeIdx < RangeNum; RangeIdx++)
+						{
+							Private::QuatFromBasisDirections(
+								Out.GetChannelRangeAttributeData(0, OutRangeOffsets[RangeIdx], RangeLengths[RangeIdx]),
+								Out.GetChannelRangeAttributeData(1, OutRangeOffsets[RangeIdx], RangeLengths[RangeIdx]),
+								Out.GetChannelRangeAttributeData(2, OutRangeOffsets[RangeIdx], RangeLengths[RangeIdx]),
+								Out.GetChannelRangeAttributeData(3, OutRangeOffsets[RangeIdx], RangeLengths[RangeIdx]),
+
+								In[0]->GetChannelRangeAttributeData(0, InRangeOffsets[0][RangeIdx], RangeLengths[RangeIdx]),
+								In[0]->GetChannelRangeAttributeData(1, InRangeOffsets[0][RangeIdx], RangeLengths[RangeIdx]),
+								In[0]->GetChannelRangeAttributeData(2, InRangeOffsets[0][RangeIdx], RangeLengths[RangeIdx]),
+
+								In[1]->GetChannelRangeAttributeData(0, InRangeOffsets[1][RangeIdx], RangeLengths[RangeIdx]),
+								In[1]->GetChannelRangeAttributeData(1, InRangeOffsets[1][RangeIdx], RangeLengths[RangeIdx]),
+								In[1]->GetChannelRangeAttributeData(2, InRangeOffsets[1][RangeIdx], RangeLengths[RangeIdx]),
+
+								In[2]->GetChannelRangeAttributeData(0, InRangeOffsets[2][RangeIdx], RangeLengths[RangeIdx]),
+								In[2]->GetChannelRangeAttributeData(1, InRangeOffsets[2][RangeIdx], RangeLengths[RangeIdx]),
+								In[2]->GetChannelRangeAttributeData(2, InRangeOffsets[2][RangeIdx], RangeLengths[RangeIdx]));
+						}
+				});
+		}
+
 		void QuatInvMul(FFrameAttribute& Out, const FFrameAttribute& Lhs, const FFrameAttribute& Rhs)
 		{
 			check(Lhs.GetChannelNum() == 4);
@@ -5133,6 +5270,40 @@ namespace UE::Learning
 			Slice(Out, In, 7, 3);
 		}
 
+		void TransformApplyConstantLocation(FFrameAttribute& Out, const FFrameAttribute& Lhs, const FVector3f& Rhs)
+		{
+			check(Lhs.GetChannelNum() == 10);
+
+			UnaryOp(Out, 3, Lhs, [Rhs](
+				UE::Learning::FFrameAttribute& Out,
+				const UE::Learning::FFrameAttribute& In,
+				const TLearningArrayView<1, const int32> RangeOffsets,
+				const TLearningArrayView<1, const int32> RangeLengths) {
+
+					Private::TransformApplyConstantLocation(
+						Out.GetAttributeData(),
+						In.GetAttributeData(),
+						Rhs);
+				});
+		}
+
+		void TransformApplyConstantDirection(FFrameAttribute& Out, const FFrameAttribute& Lhs, const FVector3f& Rhs)
+		{
+			check(Lhs.GetChannelNum() == 10);
+
+			UnaryOp(Out, 3, Lhs, [Rhs](
+				UE::Learning::FFrameAttribute& Out,
+				const UE::Learning::FFrameAttribute& In,
+				const TLearningArrayView<1, const int32> RangeOffsets,
+				const TLearningArrayView<1, const int32> RangeLengths) {
+
+					Private::TransformApplyConstantDirection(
+						Out.GetAttributeData(),
+						In.GetAttributeData(),
+						Rhs);
+				});
+		}
+
 		void TransformConstantApply(FFrameAttribute& Out, const FTransform3f& Lhs, const FFrameAttribute& Rhs)
 		{
 			check(Rhs.GetChannelNum() == 3);
diff --git a/Engine/Plugins/Experimental/LearningCore/Source/Learning/Public/LearningFrameAttribute.h b/Engine/Plugins/Experimental/LearningCore/Source/Learning/Public/LearningFrameAttribute.h
--- a/Engine/Plugins/Experimental/LearningCore/Source/Learning/Public/LearningFrameAttribute.h
+++ b/Engine/Plugins/Experimental/LearningCore/Source/Learning/Public/LearningFrameAttribute.h
@@ -464,6 +464,9 @@ namespace UE::Learning
 		/** Compute a quaternion from a rotation vector. Argument is assumed to have 3 channels. */
 		UE_API void QuatFromRotationVector(FFrameAttribute& Out, const FFrameAttribute& In);
 
+		/** Compute a quaternion from a set of basis directions. Arguments are assumed to have 3 channels. */
+		UE_API void QuatFromBasisDirections(FFrameAttribute& Out, const FFrameAttribute& Forward, const FFrameAttribute& Right, const FFrameAttribute& Up);
+
 		/** Computes the quaternion inverse multiplication of two frame attributes. Both arguments are assumed to have 4 channels. */
 		UE_API void QuatInvMul(FFrameAttribute& Out, const FFrameAttribute& Lhs, const FFrameAttribute& Rhs);
 
@@ -577,6 +580,12 @@ namespace UE::Learning
 		/** Gets the scale part of a transform. Argument is assumed to have 10 channels. */
 		UE_API void TransformScale(FFrameAttribute& Out, const FFrameAttribute& In);
 
+		/** Applies the transform to a location */
+		UE_API void TransformApplyConstantLocation(FFrameAttribute& Out, const FFrameAttribute& Lhs, const FVector3f& Rhs);
+
+		/** Applies the transform to a direction */
+		UE_API void TransformApplyConstantDirection(FFrameAttribute& Out, const FFrameAttribute& Lhs, const FVector3f& Rhs);
+
 		/** Applies the transform to a vector */
 		UE_API void TransformConstantApply(FFrameAttribute& Out, const FTransform3f& Lhs, const FFrameAttribute& Rhs);
 
