LearningCore: Added Min and Max frame attribute functions

#Synced-CL 51630663

diff --git a/Engine/Plugins/Experimental/LearningCore/Source/Learning/Private/Learning.ispc b/Engine/Plugins/Experimental/LearningCore/Source/Learning/Private/Learning.ispc
--- a/Engine/Plugins/Experimental/LearningCore/Source/Learning/Private/Learning.ispc
+++ b/Engine/Plugins/Experimental/LearningCore/Source/Learning/Private/Learning.ispc
@@ -2209,6 +2209,30 @@ export void LearningFrameAttributeDiv(
 	}
 }
 
+export void LearningFrameAttributeMax(
+	uniform float Out[],
+	uniform const float Lhs[],
+	uniform const float Rhs[],
+	uniform int ValueNum)
+{
+	foreach(ValueIdx = 0 ... ValueNum)
+	{
+		Out[ValueIdx] = max(Lhs[ValueIdx], Rhs[ValueIdx]);
+	}
+}
+
+export void LearningFrameAttributeMin(
+	uniform float Out[],
+	uniform const float Lhs[],
+	uniform const float Rhs[],
+	uniform int ValueNum)
+{
+	foreach(ValueIdx = 0 ... ValueNum)
+	{
+		Out[ValueIdx] = min(Lhs[ValueIdx], Rhs[ValueIdx]);
+	}
+}
+
 export void LearningFrameAttributeAddConstantInplace(
 	uniform float InOut[],
 	uniform const float Value,
diff --git a/Engine/Plugins/Experimental/LearningCore/Source/Learning/Private/LearningFrameAttribute.cpp b/Engine/Plugins/Experimental/LearningCore/Source/Learning/Private/LearningFrameAttribute.cpp
--- a/Engine/Plugins/Experimental/LearningCore/Source/Learning/Private/LearningFrameAttribute.cpp
+++ b/Engine/Plugins/Experimental/LearningCore/Source/Learning/Private/LearningFrameAttribute.cpp
@@ -369,6 +369,46 @@
 #endif
 		}
 
+		static inline void Max(TLearningArrayView<1, float> Out, const TLearningArrayView<1, const float> Lhs, const TLearningArrayView<1, const float> Rhs)
+		{
+			check(Out.Num() == Lhs.Num());
+			check(Out.Num() == Rhs.Num());
+
+#if UE_LEARNING_ISPC
+			ispc::LearningFrameAttributeMax(
+				Out.GetData(),
+				Lhs.GetData(),
+				Rhs.GetData(),
+				Out.Num());
+#else
+			const int32 ValueNum = Out.Num();
+			for (int32 ValueIdx = 0; ValueIdx < ValueNum; ValueIdx++)
+			{
+				Out[ValueIdx] = FMath::Max(Lhs[ValueIdx], Rhs[ValueIdx]);
+			}
+#endif
+		}
+
+		static inline void Min(TLearningArrayView<1, float> Out, const TLearningArrayView<1, const float> Lhs, const TLearningArrayView<1, const float> Rhs)
+		{
+			check(Out.Num() == Lhs.Num());
+			check(Out.Num() == Rhs.Num());
+
+#if UE_LEARNING_ISPC
+			ispc::LearningFrameAttributeMin(
+				Out.GetData(),
+				Lhs.GetData(),
+				Rhs.GetData(),
+				Out.Num());
+#else
+			const int32 ValueNum = Out.Num();
+			for (int32 ValueIdx = 0; ValueIdx < ValueNum; ValueIdx++)
+			{
+				Out[ValueIdx] = FMath::Min(Lhs[ValueIdx], Rhs[ValueIdx]);
+			}
+#endif
+		}
+
 		static inline void Dot(TLearningArrayView<1, float> InOut, const TLearningArrayView<1, const float> Lhs, const TLearningArrayView<1, const float> Rhs)
 		{
 			check(InOut.Num() == Lhs.Num());
@@ -636,16 +676,26 @@
 			}
 		}
 
-		static inline void ConstantAdd(TLearningArrayView<1, float> Out, const float Lhs, const TLearningArrayView<1, const float> Rhs)
+		static inline void MaxConstant(TLearningArrayView<1, float> Out, const TLearningArrayView<1, const float> Lhs, const float Rhs)
 		{
-			check(Out.Num() == Rhs.Num());
+			check(Out.Num() == Lhs.Num());
 			const int32 ValueNum = Out.Num();
 			for (int32 ValueIdx = 0; ValueIdx < ValueNum; ValueIdx++)
 			{
-				Out[ValueIdx] = Lhs + Rhs[ValueIdx];
+				Out[ValueIdx] = FMath::Max(Lhs[ValueIdx], Rhs);
 			}
 		}
 
+		static inline void MinConstant(TLearningArrayView<1, float> Out, const TLearningArrayView<1, const float> Lhs, const float Rhs)
+		{
+			check(Out.Num() == Lhs.Num());
+			const int32 ValueNum = Out.Num();
+			for (int32 ValueIdx = 0; ValueIdx < ValueNum; ValueIdx++)
+			{
+				Out[ValueIdx] = FMath::Min(Lhs[ValueIdx], Rhs);
+			}
+		}
+
 		static inline void ConstantSub(TLearningArrayView<1, float> Out, const float Lhs, const TLearningArrayView<1, const float> Rhs)
 		{
 			check(Out.Num() == Rhs.Num());
@@ -656,13 +706,3 @@
 			}
 		}
 
-		static inline void ConstantMul(TLearningArrayView<1, float> Out, const float Lhs, const TLearningArrayView<1, const float> Rhs)
-		{
-			check(Out.Num() == Rhs.Num());
-			const int32 ValueNum = Out.Num();
-			for (int32 ValueIdx = 0; ValueIdx < ValueNum; ValueIdx++)
-			{
-				Out[ValueIdx] = Lhs * Rhs[ValueIdx];
-			}
-		}
-
@@ -976,26 +1016,6 @@
 			}
 		}
 
-		static inline void ConstantEq(TLearningArrayView<1, float> Out, const float Lhs, const TLearningArrayView<1, const float> Rhs)
-		{
-			check(Out.Num() == Rhs.Num());
-			const int32 ValueNum = Out.Num();
-			for (int32 ValueIdx = 0; ValueIdx < ValueNum; ValueIdx++)
-			{
-				Out[ValueIdx] = Lhs == Rhs[ValueIdx];
-			}
-		}
-
-		static inline void ConstantNeq(TLearningArrayView<1, float> Out, const float Lhs, const TLearningArrayView<1, const float> Rhs)
-		{
-			check(Out.Num() == Rhs.Num());
-			const int32 ValueNum = Out.Num();
-			for (int32 ValueIdx = 0; ValueIdx < ValueNum; ValueIdx++)
-			{
-				Out[ValueIdx] = Lhs != Rhs[ValueIdx];
-			}
-		}
-
 		static inline void FilterWithKernel(TLearningArrayView<1, float> Out, const TLearningArrayView<1, const float> In, const TLearningArrayView<1, const float> Kernel)
 		{
 			check(Kernel.Num() % 2 == 1); // Kernel must be odd in length
@@ -2948,6 +2968,89 @@
 				});
 		}
 
+
+		void Max(FFrameAttribute& Out, const FFrameAttribute& Lhs, const FFrameAttribute& Rhs)
+		{
+			check(Lhs.GetChannelNum() == Rhs.GetChannelNum());
+
+			if (FrameRangeSet::Equal(Lhs.FrameRangeSet, Rhs.FrameRangeSet))
+			{
+				Out.FrameRangeSet = Lhs.FrameRangeSet;
+				Out.AttributeData.SetNumUninitialized({ Lhs.GetChannelNum(), Lhs.GetTotalFrameNum() });
+
+				Private::Max(
+					Out.GetAttributeData().Flatten(),
+					Lhs.GetAttributeData().Flatten(),
+					Rhs.GetAttributeData().Flatten());
+
+				return;
+			}
+
+			BinaryOp(Out, Lhs.GetChannelNum(), Lhs, Rhs, [](
+				FFrameAttribute& Out,
+				const FFrameAttribute& Lhs,
+				const FFrameAttribute& Rhs,
+				const TLearningArrayView<1, const int32> OutRangeOffsets,
+				const TLearningArrayView<1, const int32> LhsRangeOffsets,
+				const TLearningArrayView<1, const int32> RhsRangeOffsets,
+				const TLearningArrayView<1, const int32> RangeLengths) {
+
+					const int32 ChannelNum = Out.GetChannelNum();
+					const int32 RangeNum = RangeLengths.Num();
+					for (int32 ChannelIdx = 0; ChannelIdx < ChannelNum; ChannelIdx++)
+					{
+						for (int32 RangeIdx = 0; RangeIdx < RangeNum; RangeIdx++)
+						{
+							Private::Max(
+								Out.GetChannelRangeAttributeData(ChannelIdx, OutRangeOffsets[RangeIdx], RangeLengths[RangeIdx]),
+								Lhs.GetChannelRangeAttributeData(ChannelIdx, LhsRangeOffsets[RangeIdx], RangeLengths[RangeIdx]),
+								Rhs.GetChannelRangeAttributeData(ChannelIdx, RhsRangeOffsets[RangeIdx], RangeLengths[RangeIdx]));
+						}
+					}
+				});
+		}
+
+		void Min(FFrameAttribute& Out, const FFrameAttribute& Lhs, const FFrameAttribute& Rhs)
+		{
+			check(Lhs.GetChannelNum() == Rhs.GetChannelNum());
+
+			if (FrameRangeSet::Equal(Lhs.FrameRangeSet, Rhs.FrameRangeSet))
+			{
+				Out.FrameRangeSet = Lhs.FrameRangeSet;
+				Out.AttributeData.SetNumUninitialized({ Lhs.GetChannelNum(), Lhs.GetTotalFrameNum() });
+
+				Private::Min(
+					Out.GetAttributeData().Flatten(),
+					Lhs.GetAttributeData().Flatten(),
+					Rhs.GetAttributeData().Flatten());
+
+				return;
+			}
+
+			BinaryOp(Out, Lhs.GetChannelNum(), Lhs, Rhs, [](
+				FFrameAttribute& Out,
+				const FFrameAttribute& Lhs,
+				const FFrameAttribute& Rhs,
+				const TLearningArrayView<1, const int32> OutRangeOffsets,
+				const TLearningArrayView<1, const int32> LhsRangeOffsets,
+				const TLearningArrayView<1, const int32> RhsRangeOffsets,
+				const TLearningArrayView<1, const int32> RangeLengths) {
+
+					const int32 ChannelNum = Out.GetChannelNum();
+					const int32 RangeNum = RangeLengths.Num();
+					for (int32 ChannelIdx = 0; ChannelIdx < ChannelNum; ChannelIdx++)
+					{
+						for (int32 RangeIdx = 0; RangeIdx < RangeNum; RangeIdx++)
+						{
+							Private::Min(
+								Out.GetChannelRangeAttributeData(ChannelIdx, OutRangeOffsets[RangeIdx], RangeLengths[RangeIdx]),
+								Lhs.GetChannelRangeAttributeData(ChannelIdx, LhsRangeOffsets[RangeIdx], RangeLengths[RangeIdx]),
+								Rhs.GetChannelRangeAttributeData(ChannelIdx, RhsRangeOffsets[RangeIdx], RangeLengths[RangeIdx]));
+						}
+					}
+				});
+		}
+
 		void Dot(FFrameAttribute& Out, const FFrameAttribute& Lhs, const FFrameAttribute& Rhs)
 		{
 			check(Lhs.GetChannelNum() == Rhs.GetChannelNum());
@@ -2991,6 +3094,24 @@
 				});
 		}
 
+		void Repeat(FFrameAttribute& Out, const int32 ChannelNum, const FFrameAttribute& In)
+		{
+			check(In.GetChannelNum() == 1);
+			check(ChannelNum > 0);
+
+			UnaryOp(Out, ChannelNum, In, [ChannelNum](
+				FFrameAttribute& Out,
+				const FFrameAttribute& In,
+				const TLearningArrayView<1, const int32> RangeOffsets,
+				const TLearningArrayView<1, const int32> RangeLengths) {
+
+					for (int32 ChannelIdx = 0; ChannelIdx < ChannelNum; ChannelIdx++)
+					{
+						Private::Copy(Out.GetChannelAttributeData(ChannelIdx), In.GetChannelAttributeData(0));
+					}
+				});
+		}
+
 		void Neg(FFrameAttribute& Out, const FFrameAttribute& In)
 		{
 			UnaryOp(Out, In.GetChannelNum(), In, [](
@@ -3329,6 +3450,8 @@
 
 		void AddConstant(FFrameAttribute& Out, const FFrameAttribute& Lhs, const TLearningArrayView<1, const float> Rhs)
 		{
+			check(Lhs.GetChannelNum() == Rhs.Num());
+
 			UnaryOp(Out, Lhs.GetChannelNum(), Lhs, [Rhs](
 				FFrameAttribute& Out,
 				const FFrameAttribute& In,
@@ -3349,6 +3472,8 @@
 
 		void SubConstant(FFrameAttribute& Out, const FFrameAttribute& Lhs, const TLearningArrayView<1, const float> Rhs)
 		{
+			check(Lhs.GetChannelNum() == Rhs.Num());
+
 			UnaryOp(Out, Lhs.GetChannelNum(), Lhs, [Rhs](
 				FFrameAttribute& Out,
 				const FFrameAttribute& In,
@@ -3369,6 +3494,8 @@
 
 		void MulConstant(FFrameAttribute& Out, const FFrameAttribute& Lhs, const TLearningArrayView<1, const float> Rhs)
 		{
+			check(Lhs.GetChannelNum() == Rhs.Num());
+
 			UnaryOp(Out, Lhs.GetChannelNum(), Lhs, [Rhs](
 				FFrameAttribute& Out,
 				const FFrameAttribute& In,
@@ -3389,6 +3516,8 @@
 
 		void DivConstant(FFrameAttribute& Out, const FFrameAttribute& Lhs, const TLearningArrayView<1, const float> Rhs)
 		{
+			check(Lhs.GetChannelNum() == Rhs.Num());
+
 			UnaryOp(Out, Lhs.GetChannelNum(), Lhs, [Rhs](
 				FFrameAttribute& Out,
 				const FFrameAttribute& In,
@@ -3407,9 +3536,11 @@
 				});
 		}
 
-		void ConstantAdd(FFrameAttribute& Out, const TLearningArrayView<1, const float> Lhs, const FFrameAttribute& Rhs)
+		void MaxConstant(FFrameAttribute& Out, const FFrameAttribute& Lhs, const TLearningArrayView<1, const float> Rhs)
 		{
-			UnaryOp(Out, Rhs.GetChannelNum(), Rhs, [Lhs](
+			check(Lhs.GetChannelNum() == Rhs.Num());
+
+			UnaryOp(Out, Lhs.GetChannelNum(), Lhs, [Rhs](
 				FFrameAttribute& Out,
 				const FFrameAttribute& In,
 				const TLearningArrayView<1, const int32> RangeOffsets,
@@ -3419,17 +3550,19 @@
 
 					for (int32 ChannelIdx = 0; ChannelIdx < ChannelNum; ChannelIdx++)
 					{
-						Private::ConstantAdd(
+						Private::MaxConstant(
 							Out.GetChannelAttributeData(ChannelIdx),
-							Lhs[ChannelIdx],
-							In.GetChannelAttributeData(ChannelIdx));
+							In.GetChannelAttributeData(ChannelIdx),
+							Rhs[ChannelIdx]);
 					}
 				});
 		}
 
-		void ConstantSub(FFrameAttribute& Out, const TLearningArrayView<1, const float> Lhs, const FFrameAttribute& Rhs)
+		void MinConstant(FFrameAttribute& Out, const FFrameAttribute& Lhs, const TLearningArrayView<1, const float> Rhs)
 		{
-			UnaryOp(Out, Rhs.GetChannelNum(), Rhs, [Lhs](
+			check(Lhs.GetChannelNum() == Rhs.Num());
+
+			UnaryOp(Out, Lhs.GetChannelNum(), Lhs, [Rhs](
 				FFrameAttribute& Out,
 				const FFrameAttribute& In,
 				const TLearningArrayView<1, const int32> RangeOffsets,
@@ -3439,16 +3572,24 @@
 
 					for (int32 ChannelIdx = 0; ChannelIdx < ChannelNum; ChannelIdx++)
 					{
-						Private::ConstantSub(
+						Private::MinConstant(
 							Out.GetChannelAttributeData(ChannelIdx),
-							Lhs[ChannelIdx],
-							In.GetChannelAttributeData(ChannelIdx));
+							In.GetChannelAttributeData(ChannelIdx),
+							Rhs[ChannelIdx]);
 					}
 				});
 		}
 
-		void ConstantMul(FFrameAttribute& Out, const TLearningArrayView<1, const float> Lhs, const FFrameAttribute& Rhs)
+		void ConstantAdd(FFrameAttribute& Out, const TLearningArrayView<1, const float> Lhs, const FFrameAttribute& Rhs)
 		{
+			// Since operation is symmetric we can just swap arguments
+			return AddConstant(Out, Rhs, Lhs);
+		}
+
+		void ConstantSub(FFrameAttribute& Out, const TLearningArrayView<1, const float> Lhs, const FFrameAttribute& Rhs)
+		{
+			check(Lhs.Num() == Rhs.GetChannelNum());
+
 			UnaryOp(Out, Rhs.GetChannelNum(), Rhs, [Lhs](
 				FFrameAttribute& Out,
 				const FFrameAttribute& In,
@@ -3459,7 +3600,7 @@
 
 					for (int32 ChannelIdx = 0; ChannelIdx < ChannelNum; ChannelIdx++)
 					{
-						Private::ConstantMul(
+						Private::ConstantSub(
 							Out.GetChannelAttributeData(ChannelIdx),
 							Lhs[ChannelIdx],
 							In.GetChannelAttributeData(ChannelIdx));
@@ -3467,8 +3608,16 @@
 				});
 		}
 
+		void ConstantMul(FFrameAttribute& Out, const TLearningArrayView<1, const float> Lhs, const FFrameAttribute& Rhs)
+		{
+			// Since operation is symmetric we can just swap arguments
+			return MulConstant(Out, Rhs, Lhs);
+		}
+
 		void ConstantDiv(FFrameAttribute& Out, const TLearningArrayView<1, const float> Lhs, const FFrameAttribute& Rhs)
 		{
+			check(Lhs.Num() == Rhs.GetChannelNum());
+
 			UnaryOp(Out, Rhs.GetChannelNum(), Rhs, [Lhs](
 				FFrameAttribute& Out,
 				const FFrameAttribute& In,
@@ -3487,8 +3636,22 @@
 				});
 		}
 
+		void ConstantMax(FFrameAttribute& Out, const TLearningArrayView<1, const float> Lhs, const FFrameAttribute& Rhs)
+		{
+			// Since operation is symmetric we can just swap arguments
+			return MaxConstant(Out, Rhs, Lhs);
+		}
+
+		void ConstantMin(FFrameAttribute& Out, const TLearningArrayView<1, const float> Lhs, const FFrameAttribute& Rhs)
+		{
+			// Since operation is symmetric we can just swap arguments
+			return MinConstant(Out, Rhs, Lhs);
+		}
+
 		void ConstantDot(FFrameAttribute& Out, const TLearningArrayView<1, const float> Lhs, const FFrameAttribute& Rhs)
 		{
+			check(Lhs.Num() == Rhs.GetChannelNum());
+
 			UnaryOp(Out, 1, Rhs, [Lhs](
 				FFrameAttribute& Out,
 				const FFrameAttribute& In,
@@ -3511,6 +3674,8 @@
 
 		void ConstantSquaredDist(FFrameAttribute& Out, const TLearningArrayView<1, const float> Lhs, const FFrameAttribute& Rhs)
 		{
+			check(Lhs.Num() == Rhs.GetChannelNum());
+
 			UnaryOp(Out, 1, Rhs, [Lhs](
 				FFrameAttribute& Out,
 				const FFrameAttribute& In,
@@ -3533,6 +3698,8 @@
 
 		void ConstantDist(FFrameAttribute& Out, const TLearningArrayView<1, const float> Lhs, const FFrameAttribute& Rhs)
 		{
+			check(Lhs.Num() == Rhs.GetChannelNum());
+
 			UnaryOp(Out, 1, Rhs, [Lhs](
 				FFrameAttribute& Out,
 				const FFrameAttribute& In,
@@ -3557,6 +3724,8 @@
 
 		void AddConstantInplace(FFrameAttribute& InOut, const TLearningArrayView<1, const float> Rhs)
 		{
+			check(InOut.GetChannelNum() == Rhs.Num());
+
 			InplaceOp(InOut, [Rhs](
 				FFrameAttribute& InOut,
 				const TLearningArrayView<1, const int32> RangeOffsets,
@@ -3575,6 +3744,8 @@
 
 		void SubConstantInplace(FFrameAttribute& InOut, const TLearningArrayView<1, const float> Rhs)
 		{
+			check(InOut.GetChannelNum() == Rhs.Num());
+
 			InplaceOp(InOut, [Rhs](
 				FFrameAttribute& InOut,
 				const TLearningArrayView<1, const int32> RangeOffsets,
@@ -3593,6 +3764,8 @@
 
 		void MulConstantInplace(FFrameAttribute& InOut, const TLearningArrayView<1, const float> Rhs)
 		{
+			check(InOut.GetChannelNum() == Rhs.Num());
+
 			InplaceOp(InOut, [Rhs](
 				FFrameAttribute& InOut,
 				const TLearningArrayView<1, const int32> RangeOffsets,
@@ -3611,6 +3784,8 @@
 
 		void DivConstantInplace(FFrameAttribute& InOut, const TLearningArrayView<1, const float> Rhs)
 		{
+			check(InOut.GetChannelNum() == Rhs.Num());
+
 			InplaceOp(InOut, [Rhs](
 				FFrameAttribute& InOut,
 				const TLearningArrayView<1, const int32> RangeOffsets,
@@ -3920,6 +4095,8 @@
 
 		void GtConstant(FFrameAttribute& Out, const FFrameAttribute& Lhs, const TLearningArrayView<1, const float> Rhs)
 		{
+			check(Lhs.GetChannelNum() == Rhs.Num());
+
 			UnaryOp(Out, Lhs.GetChannelNum(), Lhs, [Rhs](
 				FFrameAttribute& Out,
 				const FFrameAttribute& In,
@@ -3940,6 +4117,8 @@
 
 		void GeConstant(FFrameAttribute& Out, const FFrameAttribute& Lhs, const TLearningArrayView<1, const float> Rhs)
 		{
+			check(Lhs.GetChannelNum() == Rhs.Num());
+
 			UnaryOp(Out, Lhs.GetChannelNum(), Lhs, [Rhs](
 				FFrameAttribute& Out,
 				const FFrameAttribute& In,
@@ -3960,6 +4139,8 @@
 
 		void LtConstant(FFrameAttribute& Out, const FFrameAttribute& Lhs, const TLearningArrayView<1, const float> Rhs)
 		{
+			check(Lhs.GetChannelNum() == Rhs.Num());
+
 			UnaryOp(Out, Lhs.GetChannelNum(), Lhs, [Rhs](
 				FFrameAttribute& Out,
 				const FFrameAttribute& In,
@@ -3980,6 +4161,8 @@
 
 		void LeConstant(FFrameAttribute& Out, const FFrameAttribute& Lhs, const TLearningArrayView<1, const float> Rhs)
 		{
+			check(Lhs.GetChannelNum() == Rhs.Num());
+
 			UnaryOp(Out, Lhs.GetChannelNum(), Lhs, [Rhs](
 				FFrameAttribute& Out,
 				const FFrameAttribute& In,
@@ -4000,6 +4183,8 @@
 
 		void EqConstant(FFrameAttribute& Out, const FFrameAttribute& Lhs, const TLearningArrayView<1, const float> Rhs)
 		{
+			check(Lhs.GetChannelNum() == Rhs.Num());
+
 			UnaryOp(Out, Lhs.GetChannelNum(), Lhs, [Rhs](
 				FFrameAttribute& Out,
 				const FFrameAttribute& In,
@@ -4020,6 +4205,8 @@
 
 		void NeqConstant(FFrameAttribute& Out, const FFrameAttribute& Lhs, const TLearningArrayView<1, const float> Rhs)
 		{
+			check(Lhs.GetChannelNum() == Rhs.Num());
+
 			UnaryOp(Out, Lhs.GetChannelNum(), Lhs, [Rhs](
 				FFrameAttribute& Out,
 				const FFrameAttribute& In,
@@ -4040,6 +4227,8 @@
 
 		void ConstantGt(FFrameAttribute& Out, const TLearningArrayView<1, const float> Lhs, const FFrameAttribute& Rhs)
 		{
+			check(Lhs.Num() == Rhs.GetChannelNum());
+
 			UnaryOp(Out, Rhs.GetChannelNum(), Rhs, [Lhs](
 				FFrameAttribute& Out,
 				const FFrameAttribute& In,
@@ -4060,6 +4249,8 @@
 
 		void ConstantGe(FFrameAttribute& Out, const TLearningArrayView<1, const float> Lhs, const FFrameAttribute& Rhs)
 		{
+			check(Lhs.Num() == Rhs.GetChannelNum());
+
 			UnaryOp(Out, Rhs.GetChannelNum(), Rhs, [Lhs](
 				FFrameAttribute& Out,
 				const FFrameAttribute& In,
@@ -4080,6 +4271,8 @@
 
 		void ConstantLt(FFrameAttribute& Out, const TLearningArrayView<1, const float> Lhs, const FFrameAttribute& Rhs)
 		{
+			check(Lhs.Num() == Rhs.GetChannelNum());
+
 			UnaryOp(Out, Rhs.GetChannelNum(), Rhs, [Lhs](
 				FFrameAttribute& Out,
 				const FFrameAttribute& In,
@@ -4100,6 +4293,8 @@
 
 		void ConstantLe(FFrameAttribute& Out, const TLearningArrayView<1, const float> Lhs, const FFrameAttribute& Rhs)
 		{
+			check(Lhs.Num() == Rhs.GetChannelNum());
+
 			UnaryOp(Out, Rhs.GetChannelNum(), Rhs, [Lhs](
 				FFrameAttribute& Out,
 				const FFrameAttribute& In,
@@ -4120,42 +4315,14 @@
 
 		void ConstantEq(FFrameAttribute& Out, const TLearningArrayView<1, const float> Lhs, const FFrameAttribute& Rhs)
 		{
-			UnaryOp(Out, Rhs.GetChannelNum(), Rhs, [Lhs](
-				FFrameAttribute& Out,
-				const FFrameAttribute& In,
-				const TLearningArrayView<1, const int32> RangeOffsets,
-				const TLearningArrayView<1, const int32> RangeLengths) {
-
-					const int32 ChannelNum = Out.GetChannelNum();
-
-					for (int32 ChannelIdx = 0; ChannelIdx < ChannelNum; ChannelIdx++)
-					{
-						Private::ConstantEq(
-							Out.GetChannelAttributeData(ChannelIdx),
-							Lhs[ChannelIdx],
-							In.GetChannelAttributeData(ChannelIdx));
-					}
-				});
+			// Since operation is symmetric we can just swap arguments
+			return EqConstant(Out, Rhs, Lhs);
 		}
 
 		void ConstantNeq(FFrameAttribute& Out, const TLearningArrayView<1, const float> Lhs, const FFrameAttribute& Rhs)
 		{
-			UnaryOp(Out, Rhs.GetChannelNum(), Rhs, [Lhs](
-				FFrameAttribute& Out,
-				const FFrameAttribute& In,
-				const TLearningArrayView<1, const int32> RangeOffsets,
-				const TLearningArrayView<1, const int32> RangeLengths) {
-
-					const int32 ChannelNum = Out.GetChannelNum();
-
-					for (int32 ChannelIdx = 0; ChannelIdx < ChannelNum; ChannelIdx++)
-					{
-						Private::ConstantNeq(
-							Out.GetChannelAttributeData(ChannelIdx),
-							Lhs[ChannelIdx],
-							In.GetChannelAttributeData(ChannelIdx));
-					}
-				});
+			// Since operation is symmetric we can just swap arguments
+			return NeqConstant(Out, Rhs, Lhs);
 		}
 
 		void FilterGaussian(FFrameAttribute& Out, const FFrameAttribute& In, const float StdInFrames)
diff --git a/Engine/Plugins/Experimental/LearningCore/Source/Learning/Public/LearningFrameAttribute.h b/Engine/Plugins/Experimental/LearningCore/Source/Learning/Public/LearningFrameAttribute.h
--- a/Engine/Plugins/Experimental/LearningCore/Source/Learning/Public/LearningFrameAttribute.h
+++ b/Engine/Plugins/Experimental/LearningCore/Source/Learning/Public/LearningFrameAttribute.h
@@ -228,6 +228,12 @@ namespace UE::Learning
 		/** Divide two frame attributes. Channel numbers must match. */
 		UE_API void Div(FFrameAttribute& Out, const FFrameAttribute& Lhs, const FFrameAttribute& Rhs);
 
+		/** Computes the component-wise maximum of two frame attributes. Channel numbers must match. */
+		UE_API void Max(FFrameAttribute& Out, const FFrameAttribute& Lhs, const FFrameAttribute& Rhs);
+
+		/** Computes the component-wise minimum of two frame attributes. Channel numbers must match. */
+		UE_API void Min(FFrameAttribute& Out, const FFrameAttribute& Lhs, const FFrameAttribute& Rhs);
+
 
 		/** Compute the dot product of two frame attributes. Channel numbers must match. */
 		UE_API void Dot(FFrameAttribute& Out, const FFrameAttribute& Lhs, const FFrameAttribute& Rhs);
@@ -235,6 +241,9 @@ namespace UE::Learning
 		/** Copy a frame attribute. */
 		UE_API void Copy(FFrameAttribute& Out, const FFrameAttribute& In);
 
+		/** Create a frame attribute filled with the given single channel frame attribute at each frame */
+		UE_API void Repeat(FFrameAttribute& Out, const int32 ChannelNum, const FFrameAttribute& In);
+
 		/** Negate a frame attribute. */
 		UE_API void Neg(FFrameAttribute& Out, const FFrameAttribute& In);
 
@@ -306,6 +315,12 @@ namespace UE::Learning
 		/** Divide a frame attribute by a constant value. Channel number must match the size of Rhs. */
 		UE_API void DivConstant(FFrameAttribute& Out, const FFrameAttribute& Lhs, const TLearningArrayView<1, const float> Rhs);
 
+		/** Compute the component-wise maximum of a constant value and a frame attribute. Channel number must match the size of Rhs. */
+		UE_API void MaxConstant(FFrameAttribute& Out, const FFrameAttribute& Lhs, const TLearningArrayView<1, const float> Rhs);
+
+		/** Compute the component-wise minimum of a constant value and a frame attribute. Channel number must match the size of Rhs. */
+		UE_API void MinConstant(FFrameAttribute& Out, const FFrameAttribute& Lhs, const TLearningArrayView<1, const float> Rhs);
+
 
 		/** Add a constant value to a frame attribute. Channel number must match the size of Lhs. */
 		UE_API void ConstantAdd(FFrameAttribute& Out, const TLearningArrayView<1, const float> Lhs, const FFrameAttribute& Rhs);
@@ -319,6 +334,13 @@ namespace UE::Learning
 		/** Divide a frame attribute by a constant value. Channel number must match the size of Lhs. */
 		UE_API void ConstantDiv(FFrameAttribute& Out, const TLearningArrayView<1, const float> Lhs, const FFrameAttribute& Rhs);
 
+		/** Compute the component-wise maximum of a constant value and a frame attribute. Channel number must match the size of Lhs. */
+		UE_API void ConstantMax(FFrameAttribute& Out, const TLearningArrayView<1, const float> Lhs, const FFrameAttribute& Rhs);
+
+		/** Compute the component-wise minimum of a constant value and a frame attribute. Channel number must match the size of Lhs. */
+		UE_API void ConstantMin(FFrameAttribute& Out, const TLearningArrayView<1, const float> Lhs, const FFrameAttribute& Rhs);
+
+
 
 		/** Gets the dot product with a constant value. Channel number must match the size of Lhs. */
 		UE_API void ConstantDot(FFrameAttribute& Out, const TLearningArrayView<1, const float> Lhs, const FFrameAttribute& Rhs);
