LearningCore: Added cross product

#Synced-CL 51630663

diff --git a/Engine/Plugins/Experimental/LearningCore/Source/Learning/Private/LearningFrameAttribute.cpp b/Engine/Plugins/Experimental/LearningCore/Source/Learning/Private/LearningFrameAttribute.cpp
--- a/Engine/Plugins/Experimental/LearningCore/Source/Learning/Private/LearningFrameAttribute.cpp
+++ b/Engine/Plugins/Experimental/LearningCore/Source/Learning/Private/LearningFrameAttribute.cpp
@@ -420,6 +420,40 @@ namespace UE::Learning
 			}
 		}
 
+		static inline void Cross(
+			TLearningArrayView<1, float> OutX,
+			TLearningArrayView<1, float> OutY,
+			TLearningArrayView<1, float> OutZ,
+			const TLearningArrayView<1, const float> LhsX,
+			const TLearningArrayView<1, const float> LhsY,
+			const TLearningArrayView<1, const float> LhsZ,
+			const TLearningArrayView<1, const float> RhsX,
+			const TLearningArrayView<1, const float> RhsY,
+			const TLearningArrayView<1, const float> RhsZ)
+		{
+			check(OutX.Num() == OutY.Num());
+			check(OutX.Num() == OutZ.Num());
+			check(OutX.Num() == LhsX.Num());
+			check(OutX.Num() == LhsY.Num());
+			check(OutX.Num() == LhsZ.Num());
+			check(OutX.Num() == RhsX.Num());
+			check(OutX.Num() == RhsY.Num());
+			check(OutX.Num() == RhsZ.Num());
+
+			const int32 ValueNum = OutX.Num();
+
+			for (int32 ValueIdx = 0; ValueIdx < ValueNum; ValueIdx++)
+			{
+				const FVector3f Lhs = FVector3f(LhsX[ValueIdx], LhsY[ValueIdx], LhsZ[ValueIdx]);
+				const FVector3f Rhs = FVector3f(RhsX[ValueIdx], RhsY[ValueIdx], RhsZ[ValueIdx]);
+				const FVector3f Out = Lhs.Cross(Rhs);
+
+				OutX[ValueIdx] = Out.X;
+				OutY[ValueIdx] = Out.Y;
+				OutZ[ValueIdx] = Out.Z;
+			}
+		}
+
 		static inline void Neg(TLearningArrayView<1, float> Out, const TLearningArrayView<1, const float> In)
 		{
 			check(Out.Num() == In.Num());
@@ -3082,6 +3116,38 @@ namespace UE::Learning
 				});
 		}
 
+		void Cross(FFrameAttribute& Out, const FFrameAttribute& Lhs, const FFrameAttribute& Rhs)
+		{
+			check(Lhs.GetChannelNum() == 3);
+			check(Rhs.GetChannelNum() == 3);
+
+			BinaryOp(Out, 3, Lhs, Rhs, [](
+				FFrameAttribute& Out,
+				const FFrameAttribute& Lhs,
+				const FFrameAttribute& Rhs,
+				const TLearningArrayView<1, const int32> OutRangeOffsets,
+				const TLearningArrayView<1, const int32> LhsRangeOffsets,
+				const TLearningArrayView<1, const int32> RhsRangeOffsets,
+				const TLearningArrayView<1, const int32> RangeLengths) {
+
+					const int32 RangeNum = RangeLengths.Num();
+
+					for (int32 RangeIdx = 0; RangeIdx < RangeNum; RangeIdx++)
+					{
+						Private::Cross(
+							Out.GetChannelRangeAttributeData(0, OutRangeOffsets[RangeIdx], RangeLengths[RangeIdx]),
+							Out.GetChannelRangeAttributeData(1, OutRangeOffsets[RangeIdx], RangeLengths[RangeIdx]),
+							Out.GetChannelRangeAttributeData(2, OutRangeOffsets[RangeIdx], RangeLengths[RangeIdx]),
+							Lhs.GetChannelRangeAttributeData(0, LhsRangeOffsets[RangeIdx], RangeLengths[RangeIdx]),
+							Lhs.GetChannelRangeAttributeData(1, LhsRangeOffsets[RangeIdx], RangeLengths[RangeIdx]),
+							Lhs.GetChannelRangeAttributeData(2, LhsRangeOffsets[RangeIdx], RangeLengths[RangeIdx]),
+							Rhs.GetChannelRangeAttributeData(0, RhsRangeOffsets[RangeIdx], RangeLengths[RangeIdx]),
+							Rhs.GetChannelRangeAttributeData(1, RhsRangeOffsets[RangeIdx], RangeLengths[RangeIdx]),
+							Rhs.GetChannelRangeAttributeData(2, RhsRangeOffsets[RangeIdx], RangeLengths[RangeIdx]));
+					}
+				});
+		}
+
 		void Copy(FFrameAttribute& Out, const FFrameAttribute& In)
 		{
 			UnaryOp(Out, In.GetChannelNum(), In, [](
diff --git a/Engine/Plugins/Experimental/LearningCore/Source/Learning/Public/LearningFrameAttribute.h b/Engine/Plugins/Experimental/LearningCore/Source/Learning/Public/LearningFrameAttribute.h
--- a/Engine/Plugins/Experimental/LearningCore/Source/Learning/Public/LearningFrameAttribute.h
+++ b/Engine/Plugins/Experimental/LearningCore/Source/Learning/Public/LearningFrameAttribute.h
@@ -238,6 +238,9 @@ namespace UE::Learning
 		/** Compute the dot product of two frame attributes. Channel numbers must match. */
 		UE_API void Dot(FFrameAttribute& Out, const FFrameAttribute& Lhs, const FFrameAttribute& Rhs);
 
+		/** Compute the cross product of two frame attributes. Both arguments are assumed to have 3 channels. */
+		UE_API void Cross(FFrameAttribute& Out, const FFrameAttribute& Lhs, const FFrameAttribute& Rhs);
+
 		/** Copy a frame attribute. */
 		UE_API void Copy(FFrameAttribute& Out, const FFrameAttribute& In);
 
