LearningCore: Fixed some minor bugs in frame attributes

#Synced-CL 52006914

diff --git a/Engine/Plugins/Experimental/LearningCore/Source/Learning/Private/LearningFrameAttribute.cpp b/Engine/Plugins/Experimental/LearningCore/Source/Learning/Private/LearningFrameAttribute.cpp
--- a/Engine/Plugins/Experimental/LearningCore/Source/Learning/Private/LearningFrameAttribute.cpp
+++ b/Engine/Plugins/Experimental/LearningCore/Source/Learning/Private/LearningFrameAttribute.cpp
@@ -607,7 +607,7 @@ namespace UE::Learning
 #endif
 		}
 
-		static inline void NormalizeInplace(TLearningArrayView<2, float> InOut)
+		static inline void NormalizeInplace(TLearningArrayView<2, float> InOut, const float Eps = UE_SMALL_NUMBER)
 		{
 			const int32 ChannelNum = InOut.Num<0>();
 			const int32 FrameNum = InOut.Num<1>();
@@ -623,7 +623,7 @@ namespace UE::Learning
 
 				for (int32 ChannelIdx = 0; ChannelIdx < ChannelNum; ChannelIdx++)
 				{
-					InOut[ChannelIdx][FrameIdx] = InOut[ChannelIdx][FrameIdx] / Length;
+					InOut[ChannelIdx][FrameIdx] = InOut[ChannelIdx][FrameIdx] / FMath::Max(Length, Eps);
 				}
 			}
 		}
@@ -646,7 +646,7 @@ namespace UE::Learning
 #endif
 		}
 
-		static inline void Normalize(TLearningArrayView<2, float> Out, const TLearningArrayView<2, const float> In)
+		static inline void Normalize(TLearningArrayView<2, float> Out, const TLearningArrayView<2, const float> In, const float Eps = UE_SMALL_NUMBER)
 		{
 			check(Out.Num<0>() == In.Num<0>());
 			check(Out.Num<1>() == In.Num<1>());
@@ -665,7 +665,7 @@ namespace UE::Learning
 
 				for (int32 ChannelIdx = 0; ChannelIdx < ChannelNum; ChannelIdx++)
 				{
-					Out[ChannelIdx][FrameIdx] = In[ChannelIdx][FrameIdx] / Length;
+					Out[ChannelIdx][FrameIdx] = In[ChannelIdx][FrameIdx] / FMath::Max(Length, UE_SMALL_NUMBER);
 				}
 			}
 		}
@@ -3100,7 +3100,7 @@ namespace UE::Learning
 
 					Private::Zero(Out.GetAttributeData().Flatten());
 
-					const int32 ChannelNum = Out.GetChannelNum();
+					const int32 ChannelNum = Lhs.GetChannelNum();
 					const int32 RangeNum = RangeLengths.Num();
 
 					for (int32 ChannelIdx = 0; ChannelIdx < ChannelNum; ChannelIdx++)
