LearningCore: A few more frame ranges and frame attribute functions.

#Synced-CL 52098865

diff --git a/Engine/Plugins/Experimental/LearningCore/Source/Learning/Private/Learning.ispc b/Engine/Plugins/Experimental/LearningCore/Source/Learning/Private/Learning.ispc
--- a/Engine/Plugins/Experimental/LearningCore/Source/Learning/Private/Learning.ispc
+++ b/Engine/Plugins/Experimental/LearningCore/Source/Learning/Private/Learning.ispc
@@ -2052,6 +2052,239 @@ export uniform int32 LearningRangesDifference(
 	return OutIndex;
 }
 
+export uniform int32 LearningRangesIntersects(
+	uniform int32 OutStarts[],
+	uniform int32 OutLengths[],
+	uniform const int32 LhsStarts[],
+	uniform const int32 LhsLengths[],
+	uniform const int32 RhsStarts[],
+	uniform const int32 RhsLengths[],
+	uniform const int LhsNum,
+	uniform const int RhsNum)
+{
+	if (LhsNum == 0) { return 0; }
+	if (RhsNum == 0) { return 0; }
+
+	if (LearningRangesEqual(LhsStarts, LhsLengths, RhsStarts, RhsLengths, LhsNum, RhsNum))
+	{
+		LearningRangesCopy(
+			OutStarts,
+			OutLengths,
+			LhsStarts,
+			LhsLengths,
+			LhsNum);
+
+		return LhsNum;
+	}
+
+	// Activation state of each list of ranges
+	uniform bool bOutActive = false;
+	uniform bool bRhsActive = false;
+	uniform bool bIntersects = false;
+
+	// Event index for each list of ranges
+	uniform int32 LhsIndex = 0;
+	uniform int32 RhsIndex = 0;
+	uniform int32 OutIndex = 0;
+
+	// While both ranges have events to process
+	while (LhsIndex < LhsNum * 2 && RhsIndex < RhsNum * 2)
+	{
+		// Are the next lhs, and rhs events active or inactive
+		uniform const bool bLhsActiveNext = LhsIndex % 2 == 0;
+		uniform const bool bRhsActiveNext = RhsIndex % 2 == 0;
+
+		// Time of the next lhs, and rhs events
+		uniform const int32 LhsT = bLhsActiveNext ? LhsStarts[LhsIndex / 2] : LhsStarts[LhsIndex / 2] + LhsLengths[LhsIndex / 2];
+		uniform const int32 RhsT = bRhsActiveNext ? RhsStarts[RhsIndex / 2] : RhsStarts[RhsIndex / 2] + RhsLengths[RhsIndex / 2];
+
+		// Event from lhs coming first
+		if (LhsT < RhsT)
+		{
+			// Activate output
+			if (!bOutActive && bLhsActiveNext)
+			{
+				bIntersects = bRhsActive;
+				bOutActive = true;
+				OutStarts[OutIndex] = LhsT;
+			}
+			// Deactivate output
+			else if (bOutActive && !bLhsActiveNext)
+			{
+				bOutActive = false;
+				if (bIntersects)
+				{
+					OutLengths[OutIndex] = LhsT - OutStarts[OutIndex];
+					OutIndex++;
+				}
+			}
+
+			LhsIndex++;
+		}
+		// Event from rhs coming first
+		else if (RhsT < LhsT)
+		{
+			// Mark as intersecting
+			if (bOutActive && bRhsActiveNext)
+			{
+				bIntersects = true;
+			}
+
+			bRhsActive = bRhsActiveNext;
+			RhsIndex++;
+		}
+		// Event from lhs and rhs coming at same time
+		else
+		{
+			// Activate output
+			if (!bOutActive && bLhsActiveNext)
+			{
+				bIntersects = bRhsActiveNext;
+				bOutActive = true;
+				OutStarts[OutIndex] = LhsT;
+			}
+			// Deactivate output
+			else if (bOutActive && !bLhsActiveNext)
+			{
+				bOutActive = false;
+				if (bIntersects)
+				{
+					OutLengths[OutIndex] = LhsT - OutStarts[OutIndex];
+					OutIndex++;
+				}
+			}
+
+			bRhsActive = bRhsActiveNext;
+			LhsIndex++; RhsIndex++;
+		}
+	}
+
+	// While the output is still active and intersecting
+	if (bOutActive && bIntersects)
+	{
+		// Next lhs event should be inactive
+		uniform const bool bLhsActiveNext = LhsIndex % 2 == 0;
+
+		// Time of the next lhs event
+		uniform const int32 LhsT = bLhsActiveNext ? LhsStarts[LhsIndex / 2] : LhsStarts[LhsIndex / 2] + LhsLengths[LhsIndex / 2];
+
+		// Finish writing range
+		OutLengths[OutIndex] = LhsT - OutStarts[OutIndex];
+		OutIndex++;
+	}
+
+	// Return number of ranges added
+	return OutIndex;
+}
+
+export uniform int32 LearningRangesFramesIntersects(
+	uniform int32 OutStarts[],
+	uniform int32 OutLengths[],
+	uniform const int32 LhsStarts[],
+	uniform const int32 LhsLengths[],
+	uniform const int32 RhsFrames[],
+	uniform const int LhsNum,
+	uniform const int RhsNum)
+{
+	if (LhsNum == 0) { return 0; }
+	if (RhsNum == 0) { return 0; }
+
+	// Activation state of each list of ranges
+	uniform bool bOutActive = false;
+	uniform bool bIntersects = false;
+
+	// Event index for each list of ranges
+	uniform int32 LhsIndex = 0;
+	uniform int32 RhsIndex = 0;
+	uniform int32 OutIndex = 0;
+
+	// While both have events to process
+	while (LhsIndex < LhsNum * 2 && RhsIndex < RhsNum)
+	{
+		// Is the next lhs event active or inactive
+		uniform const bool bLhsActiveNext = LhsIndex % 2 == 0;
+
+		// Time of the next lhs, and rhs events
+		uniform const int32 LhsT = bLhsActiveNext ? LhsStarts[LhsIndex / 2] : LhsStarts[LhsIndex / 2] + LhsLengths[LhsIndex / 2];
+		uniform const int32 RhsT = RhsFrames[RhsIndex];
+
+		// Event from lhs coming first
+		if (LhsT < RhsT)
+		{
+			// Activate output
+			if (!bOutActive && bLhsActiveNext)
+			{
+				bIntersects = false;
+				bOutActive = true;
+				OutStarts[OutIndex] = LhsT;
+			}
+			// Deactivate output
+			else if (bOutActive && !bLhsActiveNext)
+			{
+				bOutActive = false;
+				if (bIntersects)
+				{
+					OutLengths[OutIndex] = LhsT - OutStarts[OutIndex];
+					OutIndex++;
+				}
+			}
+
+			LhsIndex++;
+		}
+		// Event from rhs coming first
+		else if (RhsT < LhsT)
+		{
+			// Mark as intersecting
+			if (bOutActive)
+			{
+				bIntersects = true;
+			}
+
+			RhsIndex++;
+		}
+		// Event from lhs and rhs coming at same time
+		else
+		{
+			// Activate output
+			if (!bOutActive && bLhsActiveNext)
+			{
+				bIntersects = true;
+				bOutActive = true;
+				OutStarts[OutIndex] = LhsT;
+			}
+			// Deactivate output
+			else if (bOutActive && !bLhsActiveNext)
+			{
+				bOutActive = false;
+				if (bIntersects)
+				{
+					OutLengths[OutIndex] = LhsT - OutStarts[OutIndex];
+					OutIndex++;
+				}
+			}
+
+			LhsIndex++; RhsIndex++;
+		}
+	}
+
+	// While the output is still active and intersecting
+	if (bOutActive && bIntersects)
+	{
+		// Next lhs event should be inactive
+		uniform const bool bLhsActiveNext = LhsIndex % 2 == 0;
+
+		// Time of the next lhs event
+		uniform const int32 LhsT = bLhsActiveNext ? LhsStarts[LhsIndex / 2] : LhsStarts[LhsIndex / 2] + LhsLengths[LhsIndex / 2];
+
+		// Finish writing range
+		OutLengths[OutIndex] = LhsT - OutStarts[OutIndex];
+		OutIndex++;
+	}
+
+	// Return number of ranges added
+	return OutIndex;
+}
+
 export void LearningFilterWithKernel(
 	uniform float Out[], 
 	uniform const float In[], 
diff --git a/Engine/Plugins/Experimental/LearningCore/Source/Learning/Private/LearningFrameAttribute.cpp b/Engine/Plugins/Experimental/LearningCore/Source/Learning/Private/LearningFrameAttribute.cpp
--- a/Engine/Plugins/Experimental/LearningCore/Source/Learning/Private/LearningFrameAttribute.cpp
+++ b/Engine/Plugins/Experimental/LearningCore/Source/Learning/Private/LearningFrameAttribute.cpp
@@ -1,6 +1,7 @@
 // Copyright Epic Games, Inc. All Rights Reserved.
 
 #include "LearningFrameAttribute.h"
+#include "LearningFrameSet.h"
 
 #include "Math/Transform.h"
 
@@ -514,6 +515,37 @@ namespace UE::Learning
 			}
 		}
 
+		static inline void Sin(TLearningArrayView<1, float> Out, const TLearningArrayView<1, const float> In)
+		{
+			check(Out.Num() == In.Num());
+			const int32 ValueNum = Out.Num();
+			for (int32 ValueIdx = 0; ValueIdx < ValueNum; ValueIdx++)
+			{
+				Out[ValueIdx] = FMath::Sin(In[ValueIdx]);
+			}
+		}
+
+		static inline void Cos(TLearningArrayView<1, float> Out, const TLearningArrayView<1, const float> In)
+		{
+			check(Out.Num() == In.Num());
+			const int32 ValueNum = Out.Num();
+			for (int32 ValueIdx = 0; ValueIdx < ValueNum; ValueIdx++)
+			{
+				Out[ValueIdx] = FMath::Cos(In[ValueIdx]);
+			}
+		}
+
+		static inline void Atan2(TLearningArrayView<1, float> Out, const TLearningArrayView<1, const float> InY, const TLearningArrayView<1, const float> InX)
+		{
+			check(Out.Num() == InY.Num());
+			check(Out.Num() == InX.Num());
+			const int32 ValueNum = Out.Num();
+			for (int32 ValueIdx = 0; ValueIdx < ValueNum; ValueIdx++)
+			{
+				Out[ValueIdx] = FMath::Atan2(InY[ValueIdx], InX[ValueIdx]);
+			}
+		}
+
 		static inline void NegInplace(TLearningArrayView<1, float> InOut)
 		{
 			const int32 ValueNum = InOut.Num();
@@ -1867,6 +1899,33 @@ namespace UE::Learning
 			}
 		}
 
+		static inline void AngleToDirection3D(
+			TLearningArrayView<1, float> OutX,
+			TLearningArrayView<1, float> OutY,
+			TLearningArrayView<1, float> OutZ,
+			const TLearningArrayView<1, const float> In)
+		{
+			const int32 ValueNum = In.Num();
+			for (int32 ValueIdx = 0; ValueIdx < ValueNum; ValueIdx++)
+			{
+				OutX[ValueIdx] = FMath::Cos(In[ValueIdx]);
+				OutY[ValueIdx] = FMath::Sin(In[ValueIdx]);
+				OutZ[ValueIdx] = 0.0f;
+			}
+		}
+
+		static inline void Direction3DToAngle(
+			TLearningArrayView<1, float> Out,
+			const TLearningArrayView<1, const float> InX,
+			const TLearningArrayView<1, const float> InY)
+		{
+			const int32 ValueNum = Out.Num();
+			for (int32 ValueIdx = 0; ValueIdx < ValueNum; ValueIdx++)
+			{
+				Out[ValueIdx] = FMath::Atan2(InY[ValueIdx], InX[ValueIdx]);
+			}
+		}
+
 		static inline void InertializationDistanceConstant(
 			TLearningArrayView<1, float> OutValues,
 			const TLearningArrayView<1, const float> Locations,
@@ -3292,6 +3351,70 @@ namespace UE::Learning
 				});
 		}
 
+		void Sin(FFrameAttribute& Out, const FFrameAttribute& In)
+		{
+			UnaryOp(Out, In.GetChannelNum(), In, [](
+				FFrameAttribute& Out,
+				const FFrameAttribute& In,
+				const TLearningArrayView<1, const int32> RangeOffsets,
+				const TLearningArrayView<1, const int32> RangeLengths) {
+
+					const int32 ChannelNum = Out.GetChannelNum();
+
+					for (int32 ChannelIdx = 0; ChannelIdx < ChannelNum; ChannelIdx++)
+					{
+						Private::Sin(
+							Out.GetChannelAttributeData(ChannelIdx),
+							In.GetChannelAttributeData(ChannelIdx));
+					}
+				});
+		}
+
+		void Cos(FFrameAttribute& Out, const FFrameAttribute& In)
+		{
+			UnaryOp(Out, In.GetChannelNum(), In, [](
+				FFrameAttribute& Out,
+				const FFrameAttribute& In,
+				const TLearningArrayView<1, const int32> RangeOffsets,
+				const TLearningArrayView<1, const int32> RangeLengths) {
+
+					const int32 ChannelNum = Out.GetChannelNum();
+
+					for (int32 ChannelIdx = 0; ChannelIdx < ChannelNum; ChannelIdx++)
+					{
+						Private::Cos(
+							Out.GetChannelAttributeData(ChannelIdx),
+							In.GetChannelAttributeData(ChannelIdx));
+					}
+				});
+		}
+
+		void Atan2(FFrameAttribute& Out, const FFrameAttribute& Lhs, const FFrameAttribute& Rhs)
+		{
+			check(Lhs.GetChannelNum() == 1);
+			check(Rhs.GetChannelNum() == 1);
+
+			BinaryOp(Out, 1, Lhs, Rhs, [](
+				FFrameAttribute& Out,
+				const FFrameAttribute& Lhs,
+				const FFrameAttribute& Rhs,
+				const TLearningArrayView<1, const int32> OutRangeOffsets,
+				const TLearningArrayView<1, const int32> LhsRangeOffsets,
+				const TLearningArrayView<1, const int32> RhsRangeOffsets,
+				const TLearningArrayView<1, const int32> RangeLengths) {
+
+					const int32 RangeNum = RangeLengths.Num();
+
+					for (int32 RangeIdx = 0; RangeIdx < RangeNum; RangeIdx++)
+					{
+						Private::Atan2(
+							Out.GetChannelRangeAttributeData(0, OutRangeOffsets[RangeIdx], RangeLengths[RangeIdx]),
+							Lhs.GetChannelRangeAttributeData(0, LhsRangeOffsets[RangeIdx], RangeLengths[RangeIdx]),
+							Rhs.GetChannelRangeAttributeData(0, RhsRangeOffsets[RangeIdx], RangeLengths[RangeIdx]));
+					}
+				});
+		}
+
 		void NegInplace(FFrameAttribute& InOut)
 		{
 			InplaceOp(InOut, [](
@@ -5199,6 +5322,40 @@ namespace UE::Learning
 				});
 		}
 
+		void AngleToDirection3D(FFrameAttribute& Out, const FFrameAttribute& In)
+		{
+			check(In.GetChannelNum() == 1);
+
+			UnaryOp(Out, 3, In, [](
+				FFrameAttribute& Out,
+				const FFrameAttribute& In,
+				const TLearningArrayView<1, const int32> RangeOffsets,
+				const TLearningArrayView<1, const int32> RangeLengths) {
+
+						Private::AngleToDirection3D(
+							Out.GetChannelAttributeData(0),
+							Out.GetChannelAttributeData(1),
+							Out.GetChannelAttributeData(2),
+							In.GetChannelAttributeData(0));
+				});
+		}
+
+		void Direction3DToAngle(FFrameAttribute& Out, const FFrameAttribute& In)
+		{
+			check(In.GetChannelNum() == 3);
+
+			UnaryOp(Out, 1, In, [](
+				FFrameAttribute& Out,
+				const FFrameAttribute& In,
+				const TLearningArrayView<1, const int32> RangeOffsets,
+				const TLearningArrayView<1, const int32> RangeLengths) {
+
+					Private::Direction3DToAngle(
+						Out.GetChannelAttributeData(0),
+						In.GetChannelAttributeData(0),
+						In.GetChannelAttributeData(1));
+				});
+		}
 
 		void TransformMake(FFrameAttribute& Out, const FFrameAttribute& Location, const FFrameAttribute& Rotation, const FFrameAttribute& Scale)
 		{
@@ -5698,5 +5855,207 @@ namespace UE::Learning
 				});
 		}
 
+
+		void ExtractPhase(
+			FFrameAttribute& Out,
+			const FFrameRangeSet& FrameRangeSet,
+			const FFrameSet& ZeroPhaseFrames,
+			const FFrameSet& HalfPhaseFrames,
+			const EPhaseExtrapolationMode Extrapolation)
+		{
+			Out.FrameRangeSet = FrameRangeSet;
+			Out.AttributeData.SetNumUninitialized({ 1, FrameRangeSet.GetTotalFrameNum() });
+
+			FFrameSet LhsFrames; FrameRangeSet::Intersection(LhsFrames, ZeroPhaseFrames, FrameRangeSet);
+			FFrameSet RhsFrames; FrameRangeSet::Intersection(RhsFrames, HalfPhaseFrames, FrameRangeSet);
+
+			const int32 EntryNum = FrameRangeSet.GetEntryNum();
+
+			for (int32 EntryIdx = 0; EntryIdx < EntryNum; EntryIdx++)
+			{
+				const int32 SequenceIdx = FrameRangeSet.GetEntrySequence(EntryIdx);
+				const int32 RangeNum = FrameRangeSet.GetEntryRangeNum(EntryIdx);
+
+				const int32 LhsEntryIdx = LhsFrames.FindSequenceEntry(SequenceIdx);
+				const int32 RhsEntryIdx = RhsFrames.FindSequenceEntry(SequenceIdx);
+
+				// No Rhs frames for this sequence
+				if (RhsEntryIdx == INDEX_NONE)
+				{
+					for (int32 RangeIdx = 0; RangeIdx < RangeNum; RangeIdx++)
+					{
+						Array::Set(Out.GetChannelEntryRangeAttributeData(0, EntryIdx, RangeIdx), 0.0f);
+					}
+					continue;
+				}
+
+				// No Lhs frames for this sequence
+				if (LhsEntryIdx == INDEX_NONE)
+				{
+					for (int32 RangeIdx = 0; RangeIdx < RangeNum; RangeIdx++)
+					{
+						Array::Set(Out.GetChannelEntryRangeAttributeData(0, EntryIdx, RangeIdx), UE_PI);
+					}
+					continue;
+				}
+
+				for (int32 RangeIdx = 0; RangeIdx < RangeNum; RangeIdx++)
+				{
+					const int32 RangeStart = FrameRangeSet.GetEntryRangeStart(EntryIdx, RangeIdx);
+					const int32 RangeStop = FrameRangeSet.GetEntryRangeStop(EntryIdx, RangeIdx);
+					const int32 RangeLength = FrameRangeSet.GetEntryRangeLength(EntryIdx, RangeIdx);
+
+					int32 LhsStart = 0, RhsStart = 0;
+					int32 LhsFrameNum = 0, RhsFrameNum = 0;
+					LhsFrames.FindFramesInRange(LhsStart, LhsFrameNum, LhsEntryIdx, RangeStart, RangeLength);
+					RhsFrames.FindFramesInRange(RhsStart, RhsFrameNum, RhsEntryIdx, RangeStart, RangeLength);
+
+					int32 LhsIdx = LhsStart;
+					int32 RhsIdx = RhsStart;
+
+					// Check against no frames in range
+
+					if (RhsFrameNum == 0)
+					{
+						Array::Set(Out.GetChannelEntryRangeAttributeData(0, EntryIdx, RangeIdx), 0.0f);
+						continue;
+					}
+
+					if (LhsFrameNum == 0)
+					{
+						Array::Set(Out.GetChannelEntryRangeAttributeData(0, EntryIdx, RangeIdx), UE_PI);
+						continue;
+					}
+
+					int32 LhsFrame = LhsFrames.GetEntryFrame(LhsEntryIdx, LhsIdx);
+					int32 RhsFrame = RhsFrames.GetEntryFrame(RhsEntryIdx, RhsIdx);
+					int32 LastFrame = RangeStart;
+					int32 InitialFrame = INDEX_NONE;
+					bool bLhsActive = false;
+
+					// Set initial phase value based on first entry
+					if (LhsFrame < RhsFrame)
+					{
+						InitialFrame = LhsFrame;
+						LastFrame = LhsFrame;
+						bLhsActive = true;
+						LhsIdx++;
+					}
+					else
+					{
+						InitialFrame = RhsFrame;
+						LastFrame = RhsFrame;
+						bLhsActive = false;
+						RhsIdx++;
+					}
+
+					// While we still have both left and right events
+					while (LhsIdx < LhsStart + LhsFrameNum && RhsIdx < RhsStart + RhsFrameNum)
+					{
+						LhsFrame = LhsFrames.GetEntryFrame(LhsEntryIdx, LhsIdx);
+						RhsFrame = RhsFrames.GetEntryFrame(RhsEntryIdx, RhsIdx);
+
+						if (LhsFrame < RhsFrame)
+						{
+							if (!bLhsActive)
+							{
+								for (int32 Idx = 0; Idx < LhsFrame - LastFrame + 1; Idx++)
+								{
+									Out.GetChannelEntryRangeAttributeData(0, EntryIdx, RangeIdx)[LastFrame - RangeStart + Idx] =
+										FMath::Wrap(FMath::Lerp(UE_PI, UE_TWO_PI, (float)(Idx) / FMath::Max(LhsFrame - LastFrame, UE_SMALL_NUMBER)), -UE_PI, UE_PI);
+								}
+								LastFrame = LhsFrame;
+								bLhsActive = true;
+							}
+							LhsIdx++;
+						}
+						else
+						{
+							if (bLhsActive)
+							{
+								for (int32 Idx = 0; Idx < RhsFrame - LastFrame + 1; Idx++)
+								{
+									Out.GetChannelEntryRangeAttributeData(0, EntryIdx, RangeIdx)[LastFrame - RangeStart + Idx] =
+										FMath::Lerp(0.0f, UE_PI, (float)(Idx) / FMath::Max(RhsFrame - LastFrame, UE_SMALL_NUMBER));
+								}
+								LastFrame = RhsFrame;
+								bLhsActive = false;
+							}
+							RhsIdx++;
+						}
+					}
+
+					// Remaining Left Events
+					while (LhsIdx < LhsStart + LhsFrameNum)
+					{
+						LhsFrame = LhsFrames.GetEntryFrame(LhsEntryIdx, LhsIdx);
+
+						if (!bLhsActive)
+						{
+							for (int32 Idx = 0; Idx < LhsFrame - LastFrame + 1; Idx++)
+							{
+								Out.GetChannelEntryRangeAttributeData(0, EntryIdx, RangeIdx)[LastFrame - RangeStart + Idx] =
+									FMath::Wrap(FMath::Lerp(UE_PI, UE_TWO_PI, (float)(Idx) / FMath::Max(LhsFrame - LastFrame, UE_SMALL_NUMBER)), -UE_PI, UE_PI);
+							}
+							LastFrame = LhsFrame;
+							bLhsActive = true;
+						}
+						LhsIdx++;
+					}
+
+					// Remaining Right Events
+					while (RhsIdx < RhsStart + RhsFrameNum)
+					{
+						RhsFrame = RhsFrames.GetEntryFrame(RhsEntryIdx, RhsIdx);
+
+						if (bLhsActive)
+						{
+							for (int32 Idx = 0; Idx < RhsFrame - LastFrame + 1; Idx++)
+							{
+								Out.GetChannelEntryRangeAttributeData(0, EntryIdx, RangeIdx)[LastFrame - RangeStart + Idx] =
+									FMath::Lerp(0.0f, UE_PI, (float)(Idx) / FMath::Max(RhsFrame - LastFrame, UE_SMALL_NUMBER));
+							}
+							LastFrame = RhsFrame;
+							bLhsActive = false;
+						}
+						RhsIdx++;
+					}
+
+					// Fill Ends
+
+					if (Extrapolation == EPhaseExtrapolationMode::Repeat)
+					{
+						const float InitialValue = Out.GetChannelEntryRangeAttributeData(0, EntryIdx, RangeIdx)[InitialFrame - RangeStart];
+						Array::Set(Out.GetChannelEntryRangeAttributeData(0, EntryIdx, RangeIdx).Slice(0, InitialFrame - RangeStart), InitialValue);
+
+						const float FinalValue = Out.GetChannelEntryRangeAttributeData(0, EntryIdx, RangeIdx)[LastFrame - RangeStart];
+						Array::Set(Out.GetChannelEntryRangeAttributeData(0, EntryIdx, RangeIdx).Slice(LastFrame - RangeStart, RangeStop - LastFrame), FinalValue);
+					}
+					else if (Extrapolation == EPhaseExtrapolationMode::Extrapolate)
+					{
+						const float InitialValue0 = Out.GetChannelEntryRangeAttributeData(0, EntryIdx, RangeIdx)[InitialFrame + 0 - RangeStart];
+						const float InitialValue1 = Out.GetChannelEntryRangeAttributeData(0, EntryIdx, RangeIdx)[InitialFrame + 1 - RangeStart];
+						for (int32 Idx = 0; Idx < InitialFrame - RangeStart; Idx++)
+						{
+							Out.GetChannelEntryRangeAttributeData(0, EntryIdx, RangeIdx)[Idx] = FMath::Wrap(
+								InitialValue0 - (InitialFrame - RangeStart - Idx) * (InitialValue1 - InitialValue0), -UE_PI, UE_PI);
+						}
+						
+						const float FinalValue0 = Out.GetChannelEntryRangeAttributeData(0, EntryIdx, RangeIdx)[LastFrame - 1 - RangeStart];
+						const float FinalValue1 = Out.GetChannelEntryRangeAttributeData(0, EntryIdx, RangeIdx)[LastFrame - 0 - RangeStart];
+						for (int32 Idx = 0; Idx < RangeStop - LastFrame; Idx++)
+						{
+							Out.GetChannelEntryRangeAttributeData(0, EntryIdx, RangeIdx)[LastFrame - RangeStart + Idx] = FMath::Wrap(
+								FinalValue0 + (Idx + 1) * (FinalValue1 - FinalValue0), -UE_PI, UE_PI);
+						}
+					}
+					else
+					{
+						checkNoEntry();
+					}
+				}
+			}
+		}
+
 	}
 }
diff --git a/Engine/Plugins/Experimental/LearningCore/Source/Learning/Private/LearningFrameRangeSet.cpp b/Engine/Plugins/Experimental/LearningCore/Source/Learning/Private/LearningFrameRangeSet.cpp
--- a/Engine/Plugins/Experimental/LearningCore/Source/Learning/Private/LearningFrameRangeSet.cpp
+++ b/Engine/Plugins/Experimental/LearningCore/Source/Learning/Private/LearningFrameRangeSet.cpp
@@ -1623,6 +1623,286 @@ namespace UE::Learning
 			return OutIndex;
 #endif
 		}
+
+		static inline int32 RangesIntersects(
+			TLearningArrayView<1, int32> OutStarts,
+			TLearningArrayView<1, int32> OutLengths,
+			const TLearningArrayView<1, const int32> LhsStarts,
+			const TLearningArrayView<1, const int32> LhsLengths,
+			const TLearningArrayView<1, const int32> RhsStarts,
+			const TLearningArrayView<1, const int32> RhsLengths)
+		{
+			TRACE_CPUPROFILER_EVENT_SCOPE(Learning::FrameRangeSet::Private::RangesIntersects);
+
+#if UE_LEARNING_ISPC
+			return ispc::LearningRangesIntersects(
+				OutStarts.GetData(),
+				OutLengths.GetData(),
+				LhsStarts.GetData(),
+				LhsLengths.GetData(),
+				RhsStarts.GetData(),
+				RhsLengths.GetData(),
+				LhsStarts.Num(),
+				RhsStarts.Num());
+#else
+			RangesCheck(LhsStarts, LhsLengths);
+			RangesCheck(RhsStarts, RhsLengths);
+
+			if (LhsStarts.IsEmpty()) { return 0; }
+			if (RhsStarts.IsEmpty()) { return 0; }
+
+			if (RangesEqual(LhsStarts, LhsLengths, RhsStarts, RhsLengths))
+			{
+				RangesCopy(
+					OutStarts.Slice(0, LhsStarts.Num()),
+					OutLengths.Slice(0, LhsLengths.Num()),
+					LhsStarts,
+					LhsLengths);
+
+				return LhsLengths.Num();
+			}
+
+			// Number of ranges in lhs and rhs
+			const int32 LhsNum = LhsStarts.Num();
+			const int32 RhsNum = RhsStarts.Num();
+
+			// Activation state of each list of ranges
+			bool bOutActive = false;
+			bool bRhsActive = false;
+			bool bIntersects = false;
+
+			// Event index for each list of ranges
+			int32 LhsIndex = 0;
+			int32 RhsIndex = 0;
+			int32 OutIndex = 0;
+
+			// While both ranges have events to process
+			while (LhsIndex < LhsNum * 2 && RhsIndex < RhsNum * 2)
+			{
+				// Are the next lhs, and rhs events active or inactive
+				const bool bLhsActiveNext = LhsIndex % 2 == 0;
+				const bool bRhsActiveNext = RhsIndex % 2 == 0;
+
+				// Time of the next lhs, and rhs events
+				const int32 LhsT = bLhsActiveNext ? LhsStarts[LhsIndex / 2] : LhsStarts[LhsIndex / 2] + LhsLengths[LhsIndex / 2];
+				const int32 RhsT = bRhsActiveNext ? RhsStarts[RhsIndex / 2] : RhsStarts[RhsIndex / 2] + RhsLengths[RhsIndex / 2];
+
+				// Event from lhs coming first
+				if (LhsT < RhsT)
+				{
+					// Activate output
+					if (!bOutActive && bLhsActiveNext)
+					{
+						bIntersects = bRhsActive;
+						bOutActive = true;
+						OutStarts[OutIndex] = LhsT;
+					}
+					// Deactivate output
+					else if (bOutActive && !bLhsActiveNext)
+					{
+						bOutActive = false;
+						if (bIntersects)
+						{
+							OutLengths[OutIndex] = LhsT - OutStarts[OutIndex];
+							OutIndex++;
+						}
+					}
+
+					LhsIndex++;
+				}
+				// Event from rhs coming first
+				else if (RhsT < LhsT)
+				{
+					// Mark as intersecting
+					if (bOutActive && bRhsActiveNext)
+					{
+						bIntersects = true;
+					}
+
+					bRhsActive = bRhsActiveNext;
+					RhsIndex++;
+				}
+				// Event from lhs and rhs coming at same time
+				else
+				{
+					check(LhsT == RhsT);
+
+					// Activate output
+					if (!bOutActive && bLhsActiveNext)
+					{
+						bIntersects = bRhsActiveNext;
+						bOutActive = true;
+						OutStarts[OutIndex] = LhsT;
+					}
+					// Deactivate output
+					else if (bOutActive && !bLhsActiveNext)
+					{
+						bOutActive = false;
+						if (bIntersects)
+						{
+							OutLengths[OutIndex] = LhsT - OutStarts[OutIndex];
+							OutIndex++;
+						}
+					}
+
+					bRhsActive = bRhsActiveNext;
+					LhsIndex++; RhsIndex++;
+				}
+			}
+
+			// While the output is still active and intersecting
+			if (bOutActive && bIntersects)
+			{
+				// Next lhs event should be inactive
+				const bool bLhsActiveNext = LhsIndex % 2 == 0;
+				check(!bLhsActiveNext);
+
+				// Time of the next lhs event
+				const int32 LhsT = bLhsActiveNext ? LhsStarts[LhsIndex / 2] : LhsStarts[LhsIndex / 2] + LhsLengths[LhsIndex / 2];
+				
+				// Finish writing range
+				OutLengths[OutIndex] = LhsT - OutStarts[OutIndex];
+				OutIndex++;
+			}
+
+			RangesCheck(OutStarts.Slice(0, OutIndex), OutLengths.Slice(0, OutIndex));
+
+			// Return number of ranges added
+			return OutIndex;
+#endif
+		}
+	
+		static inline int32 RangesFramesIntersects(
+			TLearningArrayView<1, int32> OutStarts,
+			TLearningArrayView<1, int32> OutLengths,
+			const TLearningArrayView<1, const int32> LhsStarts,
+			const TLearningArrayView<1, const int32> LhsLengths,
+			const TLearningArrayView<1, const int32> RhsFrames)
+		{
+			TRACE_CPUPROFILER_EVENT_SCOPE(Learning::FrameRangeSet::Private::RangesFramesIntersects);
+
+#if UE_LEARNING_ISPC
+			return ispc::LearningRangesFramesIntersects(
+				OutStarts.GetData(),
+				OutLengths.GetData(),
+				LhsStarts.GetData(),
+				LhsLengths.GetData(),
+				RhsFrames.GetData(),
+				LhsStarts.Num(),
+				RhsFrames.Num());
+#else
+			RangesCheck(LhsStarts, LhsLengths);
+			FramesCheck(RhsFrames);
+
+			if (LhsStarts.IsEmpty()) { return 0; }
+			if (RhsFrames.IsEmpty()) { return 0; }
+
+			// Number of ranges in lhs and rhs
+			const int32 LhsNum = LhsStarts.Num();
+			const int32 RhsNum = RhsFrames.Num();
+
+			// Activation state of each list of ranges
+			bool bOutActive = false;
+			bool bIntersects = false;
+
+			// Event index for each list of ranges
+			int32 LhsIndex = 0;
+			int32 RhsIndex = 0;
+			int32 OutIndex = 0;
+
+			// While both have events to process
+			while (LhsIndex < LhsNum * 2 && RhsIndex < RhsNum)
+			{
+				// Is the next lhs event active or inactive
+				const bool bLhsActiveNext = LhsIndex % 2 == 0;
+
+				// Time of the next lhs, and rhs events
+				const int32 LhsT = bLhsActiveNext ? LhsStarts[LhsIndex / 2] : LhsStarts[LhsIndex / 2] + LhsLengths[LhsIndex / 2];
+				const int32 RhsT = RhsFrames[RhsIndex];
+
+				// Event from lhs coming first
+				if (LhsT < RhsT)
+				{
+					// Activate output
+					if (!bOutActive && bLhsActiveNext)
+					{
+						bIntersects = false;
+						bOutActive = true;
+						OutStarts[OutIndex] = LhsT;
+					}
+					// Deactivate output
+					else if (bOutActive && !bLhsActiveNext)
+					{
+						bOutActive = false;
+						if (bIntersects)
+						{
+							OutLengths[OutIndex] = LhsT - OutStarts[OutIndex];
+							OutIndex++;
+						}
+					}
+
+					LhsIndex++;
+				}
+				// Event from rhs coming first
+				else if (RhsT < LhsT)
+				{
+					// Mark as intersecting
+					if (bOutActive)
+					{
+						bIntersects = true;
+					}
+
+					RhsIndex++;
+				}
+				// Event from lhs and rhs coming at same time
+				else
+				{
+					check(LhsT == RhsT);
+
+					// Activate output
+					if (!bOutActive && bLhsActiveNext)
+					{
+						bIntersects = true;
+						bOutActive = true;
+						OutStarts[OutIndex] = LhsT;
+					}
+					// Deactivate output
+					else if (bOutActive && !bLhsActiveNext)
+					{
+						bOutActive = false;
+						if (bIntersects)
+						{
+							OutLengths[OutIndex] = LhsT - OutStarts[OutIndex];
+							OutIndex++;
+						}
+					}
+
+					LhsIndex++; RhsIndex++;
+				}
+			}
+
+			// While the output is still active and intersecting
+			if (bOutActive && bIntersects)
+			{
+				// Next lhs event should be inactive
+				const bool bLhsActiveNext = LhsIndex % 2 == 0;
+				check(!bLhsActiveNext);
+
+				// Time of the next lhs event
+				const int32 LhsT = bLhsActiveNext ? LhsStarts[LhsIndex / 2] : LhsStarts[LhsIndex / 2] + LhsLengths[LhsIndex / 2];
+				
+				// Finish writing range
+				OutLengths[OutIndex] = LhsT - OutStarts[OutIndex];
+				OutIndex++;
+			}
+
+			RangesCheck(OutStarts.Slice(0, OutIndex), OutLengths.Slice(0, OutIndex));
+
+			// Return number of ranges added
+			return OutIndex;
+#endif
+		}
+
 	}
 
 	void FFrameRangeSet::Check() const
@@ -2911,6 +3191,222 @@ namespace UE::Learning
 			return RangeIndex;
 		}
 
+		void Bounds(FFrameRangeSet& Out, const FFrameRangeSet& FrameRangeSet)
+		{
+			const int32 EntryNum = FrameRangeSet.GetEntryNum();
+
+			Out.EntrySequences.SetNumUninitialized({ EntryNum });
+			Out.EntryRangeOffsets.SetNumUninitialized({ EntryNum });
+			Out.EntryRangeNums.SetNumUninitialized({ EntryNum });
+			Out.RangeStarts.SetNumUninitialized({ EntryNum });
+			Out.RangeLengths.SetNumUninitialized({ EntryNum });
+			
+			for (int32 EntryIdx = 0; EntryIdx < EntryNum; EntryIdx++)
+			{
+				Out.EntrySequences[EntryIdx] = FrameRangeSet.GetEntrySequence(EntryIdx);
+				Out.EntryRangeOffsets[EntryIdx] = EntryIdx;
+				Out.EntryRangeNums[EntryIdx] = 1;
+				Out.RangeStarts[EntryIdx] = FrameRangeSet.GetEntryRangeStart(EntryIdx, 0);
+				Out.RangeLengths[EntryIdx] = 
+					FrameRangeSet.GetEntryRangeStop(EntryIdx, FrameRangeSet.GetEntryRangeNum(EntryIdx) - 1) - 
+					FrameRangeSet.GetEntryRangeStart(EntryIdx, 0);
+			}
+			
+			Out.RangeOffsets.SetNumUninitialized({ EntryNum });
+			Private::ComputeRangeOffsets(Out.RangeOffsets, Out.RangeLengths);
+			Out.Check();
+		}
+
+		void Bounds(FFrameRangeSet& Out, const FFrameSet& FrameSet)
+		{
+			const int32 EntryNum = FrameSet.GetEntryNum();
+
+			Out.EntrySequences.SetNumUninitialized({ EntryNum });
+			Out.EntryRangeOffsets.SetNumUninitialized({ EntryNum });
+			Out.EntryRangeNums.SetNumUninitialized({ EntryNum });
+			Out.RangeStarts.SetNumUninitialized({ EntryNum });
+			Out.RangeLengths.SetNumUninitialized({ EntryNum });
+
+			for (int32 EntryIdx = 0; EntryIdx < EntryNum; EntryIdx++)
+			{
+				Out.EntrySequences[EntryIdx] = FrameSet.GetEntrySequence(EntryIdx);
+				Out.EntryRangeOffsets[EntryIdx] = EntryIdx;
+				Out.EntryRangeNums[EntryIdx] = 1;
+				Out.RangeStarts[EntryIdx] = FrameSet.GetEntryFrame(EntryIdx, 0);
+				Out.RangeLengths[EntryIdx] =
+					FrameSet.GetEntryFrame(EntryIdx, FrameSet.GetEntryFrameNum(EntryIdx) - 1) + 1 -
+					FrameSet.GetEntryFrame(EntryIdx, 0);
+			}
+
+			Out.RangeOffsets.SetNumUninitialized({ EntryNum });
+			Private::ComputeRangeOffsets(Out.RangeOffsets, Out.RangeLengths);
+			Out.Check();
+		}
+
+		void Intersects(FFrameRangeSet& Out, const FFrameRangeSet& Lhs, const FFrameRangeSet& Rhs)
+		{
+			TRACE_CPUPROFILER_EVENT_SCOPE(Learning::FrameRangeSet::Intersects);
+
+			Lhs.Check(); Rhs.Check();
+
+			if (Lhs.IsEmpty()) { Out.Empty(); return; }
+			if (Rhs.IsEmpty()) { Out.Empty(); return; }
+			if (Equal(Lhs, Rhs)) { Out = Lhs; return; }
+
+			const int32 LhsEntryNum = Lhs.GetEntryNum();
+			const int32 RhsEntryNum = Rhs.GetEntryNum();
+			const int32 LhsRangeNum = Lhs.GetTotalRangeNum();
+
+			// Allocate potential maximum number of entries and ranges we might to output
+			Out.EntrySequences.SetNumUninitialized({ LhsEntryNum });
+			Out.EntryRangeOffsets.SetNumUninitialized({ LhsEntryNum });
+			Out.EntryRangeNums.SetNumUninitialized({ LhsEntryNum });
+			Out.RangeStarts.SetNumUninitialized({ LhsRangeNum });
+			Out.RangeLengths.SetNumUninitialized({ LhsRangeNum });
+
+			// Anim index for each list of ranges
+			int32 OutIndex = 0;
+			int32 LhsIndex = 0;
+			int32 RhsIndex = 0;
+
+			// Output ranges index
+			int32 RangeIndex = 0;
+
+			// While both sets have entries
+			while (LhsIndex < LhsEntryNum && RhsIndex < RhsEntryNum)
+			{
+				// If entry is in lhs but not rhs skip
+				if (Lhs.GetEntrySequence(LhsIndex) < Rhs.GetEntrySequence(RhsIndex))
+				{
+					LhsIndex++;
+				}
+				// If entry is in rhs but not lhs skip
+				else if (Rhs.GetEntrySequence(RhsIndex) < Lhs.GetEntrySequence(LhsIndex))
+				{
+					RhsIndex++;
+				}
+				// If entry is in both lhs and rhs
+				else
+				{
+					check(Lhs.GetEntryRangeNum(LhsIndex) > 0);
+					check(Rhs.GetEntryRangeNum(RhsIndex) > 0);
+
+					// Append intersecting frames to output
+					const int32 RangeNum = FrameRangeSet::Private::RangesIntersects(
+						Out.RangeStarts.Slice(RangeIndex, Out.RangeStarts.Num() - RangeIndex),
+						Out.RangeLengths.Slice(RangeIndex, Out.RangeLengths.Num() - RangeIndex),
+						Lhs.GetEntryRangeStarts(LhsIndex),
+						Lhs.GetEntryRangeLengths(LhsIndex),
+						Rhs.GetEntryRangeStarts(RhsIndex),
+						Rhs.GetEntryRangeLengths(RhsIndex));
+
+					check(RangeNum <= Lhs.GetEntryRangeNum(LhsIndex));
+
+					if (RangeNum > 0)
+					{
+						Out.EntrySequences[OutIndex] = Lhs.GetEntrySequence(LhsIndex);
+						Out.EntryRangeOffsets[OutIndex] = RangeIndex;
+						Out.EntryRangeNums[OutIndex] = RangeNum;
+						RangeIndex += RangeNum;
+						OutIndex++;
+					}
+
+					LhsIndex++; RhsIndex++;
+				}
+			}
+
+			// Resize output to match what was added
+			Out.EntrySequences.SetNumUninitialized({ OutIndex });
+			Out.EntryRangeOffsets.SetNumUninitialized({ OutIndex });
+			Out.EntryRangeNums.SetNumUninitialized({ OutIndex });
+			Out.RangeStarts.SetNumUninitialized({ RangeIndex });
+			Out.RangeLengths.SetNumUninitialized({ RangeIndex });
+			Out.RangeOffsets.SetNumUninitialized({ RangeIndex });
+			Private::ComputeRangeOffsets(Out.RangeOffsets, Out.RangeLengths);
+			Out.Check();
+		}
+
+		void Intersects(FFrameRangeSet& Out, const FFrameRangeSet& Lhs, const FFrameSet& Rhs)
+		{
+			TRACE_CPUPROFILER_EVENT_SCOPE(Learning::FrameRangeSet::Intersects);
+
+			Lhs.Check(); Rhs.Check();
+
+			if (Lhs.IsEmpty()) { Out.Empty(); return; }
+			if (Rhs.IsEmpty()) { Out.Empty(); return; }
+
+			const int32 LhsEntryNum = Lhs.GetEntryNum();
+			const int32 RhsEntryNum = Rhs.GetEntryNum();
+			const int32 LhsRangeNum = Lhs.GetTotalRangeNum();
+
+			// Allocate potential maximum number of entries and ranges we might to output
+			Out.EntrySequences.SetNumUninitialized({ LhsEntryNum });
+			Out.EntryRangeOffsets.SetNumUninitialized({ LhsEntryNum });
+			Out.EntryRangeNums.SetNumUninitialized({ LhsEntryNum });
+			Out.RangeStarts.SetNumUninitialized({ LhsRangeNum });
+			Out.RangeLengths.SetNumUninitialized({ LhsRangeNum });
+
+			// Anim index for each list of ranges
+			int32 OutIndex = 0;
+			int32 LhsIndex = 0;
+			int32 RhsIndex = 0;
+
+			// Output ranges index
+			int32 RangeIndex = 0;
+
+			// While both sets have entries
+			while (LhsIndex < LhsEntryNum && RhsIndex < RhsEntryNum)
+			{
+				// If entry is in lhs but not rhs skip
+				if (Lhs.GetEntrySequence(LhsIndex) < Rhs.GetEntrySequence(RhsIndex))
+				{
+					LhsIndex++;
+				}
+				// If entry is in rhs but not lhs skip
+				else if (Rhs.GetEntrySequence(RhsIndex) < Lhs.GetEntrySequence(LhsIndex))
+				{
+					RhsIndex++;
+				}
+				// If entry is in both lhs and rhs
+				else
+				{
+					check(Lhs.GetEntryRangeNum(LhsIndex) > 0);
+					check(Rhs.GetEntryFrameNum(RhsIndex) > 0);
+
+					// Append intersecting frames to output
+					const int32 RangeNum = FrameRangeSet::Private::RangesFramesIntersects(
+						Out.RangeStarts.Slice(RangeIndex, Out.RangeStarts.Num() - RangeIndex),
+						Out.RangeLengths.Slice(RangeIndex, Out.RangeLengths.Num() - RangeIndex),
+						Lhs.GetEntryRangeStarts(LhsIndex),
+						Lhs.GetEntryRangeLengths(LhsIndex),
+						Rhs.GetEntryFrames(RhsIndex));
+
+					check(RangeNum <= Lhs.GetEntryRangeNum(LhsIndex));
+
+					if (RangeNum > 0)
+					{
+						Out.EntrySequences[OutIndex] = Lhs.GetEntrySequence(LhsIndex);
+						Out.EntryRangeOffsets[OutIndex] = RangeIndex;
+						Out.EntryRangeNums[OutIndex] = RangeNum;
+						RangeIndex += RangeNum;
+						OutIndex++;
+					}
+
+					LhsIndex++; RhsIndex++;
+				}
+			}
+
+			// Resize output to match what was added
+			Out.EntrySequences.SetNumUninitialized({ OutIndex });
+			Out.EntryRangeOffsets.SetNumUninitialized({ OutIndex });
+			Out.EntryRangeNums.SetNumUninitialized({ OutIndex });
+			Out.RangeStarts.SetNumUninitialized({ RangeIndex });
+			Out.RangeLengths.SetNumUninitialized({ RangeIndex });
+			Out.RangeOffsets.SetNumUninitialized({ RangeIndex });
+			Private::ComputeRangeOffsets(Out.RangeOffsets, Out.RangeLengths);
+			Out.Check();
+		}
+
 		void TrimStart(FFrameRangeSet& Out, const FFrameRangeSet& FrameRangeSet, const int32 TrimFrameNum)
 		{
 			Trim(Out, FrameRangeSet, TrimFrameNum, 0);
diff --git a/Engine/Plugins/Experimental/LearningCore/Source/Learning/Private/LearningFrameSet.cpp b/Engine/Plugins/Experimental/LearningCore/Source/Learning/Private/LearningFrameSet.cpp
--- a/Engine/Plugins/Experimental/LearningCore/Source/Learning/Private/LearningFrameSet.cpp
+++ b/Engine/Plugins/Experimental/LearningCore/Source/Learning/Private/LearningFrameSet.cpp
@@ -378,6 +378,30 @@ namespace UE::Learning
 		return false;
 	}
 
+	bool FFrameSet::FindFramesInRange(int32& OutStartFrameIdx, int32& OutFrameNum, const int32 EntryIdx, const int32 RangeStart, const int32 RangeLength) const
+	{	
+		OutStartFrameIdx = INDEX_NONE;
+		OutFrameNum = 0;
+
+		const int32 FrameNum = GetEntryFrameNum(EntryIdx);
+
+		for (int32 FrameIdx = 0; FrameIdx < FrameNum; FrameIdx++)
+		{
+			const int32 Frame = GetEntryFrame(EntryIdx, FrameIdx);
+
+			if (Frame >= RangeStart && Frame < RangeStart + RangeLength)
+			{
+				if (OutStartFrameIdx == INDEX_NONE)
+				{
+					OutStartFrameIdx = FrameIdx;
+				}
+				OutFrameNum++;
+			}
+		}
+
+		return OutStartFrameIdx != INDEX_NONE;
+	}
+
 	void FFrameSet::AddEntry(
 		const int32 InSequence,
 		const TLearningArrayView<1, const int32> InFrames)
diff --git a/Engine/Plugins/Experimental/LearningCore/Source/Learning/Public/LearningFrameAttribute.h b/Engine/Plugins/Experimental/LearningCore/Source/Learning/Public/LearningFrameAttribute.h
--- a/Engine/Plugins/Experimental/LearningCore/Source/Learning/Public/LearningFrameAttribute.h
+++ b/Engine/Plugins/Experimental/LearningCore/Source/Learning/Public/LearningFrameAttribute.h
@@ -12,7 +12,6 @@
 
 namespace UE::Learning
 {
-
 	/**
 	 * A FFrameAttribute represents a attribute associated with every frame in a FFrameRangeSet. An attribute is made up of multiple "channels" 
 	 * such as the X, Y, Z components of a location.
@@ -265,6 +264,14 @@ namespace UE::Learning
 		/** Compute the sqrt of a frame attribute. */
 		UE_API void Sqrt(FFrameAttribute& Out, const FFrameAttribute& In);
 
+		/** Compute the sin of a frame attribute. */
+		UE_API void Sin(FFrameAttribute& Out, const FFrameAttribute& In);
+
+		/** Compute the cos of a frame attribute. */
+		UE_API void Cos(FFrameAttribute& Out, const FFrameAttribute& In);
+
+		/** Compute the atan2 of two frame attributes. */
+		UE_API void Atan2(FFrameAttribute& Out, const FFrameAttribute& Lhs, const FFrameAttribute& Rhs);
 
 		/** Compute the negation of a frame attribute in place. */
 		UE_API void NegInplace(FFrameAttribute& InOut);
@@ -568,6 +575,12 @@ namespace UE::Learning
 		/** Wraps a frame attribute in-place into the range -UE_PI, UE_PI */
 		UE_API void WrapAngleInplace(FFrameAttribute& InOut);
 
+		/** Computes a 3D direction from a scalar angle, leaving the Z component as zero */
+		UE_API void AngleToDirection3D(FFrameAttribute& Out, const FFrameAttribute& In);
+
+		/** Computes an angle from a 3D direction ignoring the Z component */
+		UE_API void Direction3DToAngle(FFrameAttribute& Out, const FFrameAttribute& In);
+
 
 		/** Make a transform attribute from a location, rotation, and scale */
 		UE_API void TransformMake(FFrameAttribute& Out, const FFrameAttribute& Location, const FFrameAttribute& Rotation, const FFrameAttribute& Scale);
@@ -637,6 +650,24 @@ namespace UE::Learning
 
 		/** Gets angle traveled of each range in a quat frame attribute */
 		UE_API void QuatAngleTraveled(TArrayView<float> OutAnglesTraveled, const FFrameAttribute& In);
+
+		/* Mode for handling edges in phase extraction */
+		enum class EPhaseExtrapolationMode : uint8
+		{
+			// Repeats the phase value at the start and end frames 
+			Repeat = 0,
+
+			// Extrapolates the phase value using the nearest pair of frames
+			Extrapolate = 1,
+		};
+
+		/** Extract the phase for a given set of frame ranges, using the labeled events */
+		UE_API void ExtractPhase(
+			FFrameAttribute& Out, 
+			const FFrameRangeSet& FrameRangeSet, 
+			const FFrameSet& ZeroPhaseFrames, 
+			const FFrameSet& HalfPhaseFrames,
+			const EPhaseExtrapolationMode Extrapolation = EPhaseExtrapolationMode::Repeat);
 	}
 }
 
diff --git a/Engine/Plugins/Experimental/LearningCore/Source/Learning/Public/LearningFrameRangeSet.h b/Engine/Plugins/Experimental/LearningCore/Source/Learning/Public/LearningFrameRangeSet.h
--- a/Engine/Plugins/Experimental/LearningCore/Source/Learning/Public/LearningFrameRangeSet.h
+++ b/Engine/Plugins/Experimental/LearningCore/Source/Learning/Public/LearningFrameRangeSet.h
@@ -215,6 +215,18 @@ namespace UE::Learning
 			const FFrameRangeSet& Lhs,
 			const FFrameRangeSet& Rhs);
 
+		/** Computes the bounding ranges of a given frame range set */
+		UE_API void Bounds(FFrameRangeSet& Out, const FFrameRangeSet& FrameRangeSet);
+
+		/** Computes the bounding ranges of a given frame set */
+		UE_API void Bounds(FFrameRangeSet& Out, const FFrameSet& FrameSet);
+
+		/** Computes all of the ranges in the Lhs which contain any part of the ranges in the Rhs */
+		UE_API void Intersects(FFrameRangeSet& Out, const FFrameRangeSet& Lhs, const FFrameRangeSet& Rhs);
+
+		/** Computes all of the ranges in the Lhs which contain any part of the frames in the Rhs */
+		UE_API void Intersects(FFrameRangeSet& Out, const FFrameRangeSet& Lhs, const FFrameSet& Rhs);
+		
 		/** Creates a new frame range set with the start of all ranges trimmed by the given number of frames */
 		UE_API void TrimStart(FFrameRangeSet& Out, const FFrameRangeSet& FrameRangeSet, const int32 TrimFrameNum);
 
diff --git a/Engine/Plugins/Experimental/LearningCore/Source/Learning/Public/LearningFrameSet.h b/Engine/Plugins/Experimental/LearningCore/Source/Learning/Public/LearningFrameSet.h
--- a/Engine/Plugins/Experimental/LearningCore/Source/Learning/Public/LearningFrameSet.h
+++ b/Engine/Plugins/Experimental/LearningCore/Source/Learning/Public/LearningFrameSet.h
@@ -106,6 +106,9 @@ namespace UE::Learning
 		/** Finds entry index and frame index associated with some flat offset */
 		UE_API bool FindOffset(int32& OutEntryIdx, int32& OutFrameIdx, const int32 Offset) const;
 
+		/** Finds all the frames within the given range, returning the first frame index and the number of frames */
+		UE_API bool FindFramesInRange(int32& OutStartFrameIdx, int32& OutFrameNum, const int32 EntryIdx, const int32 RangeStart, const int32 RangeLength) const;
+
 	public:
 
 		/** Array of sequences associated with each entry */
