LearningCore: Made platform specific guid string conversion method.

#Synced-CL 52297216

diff --git a/Engine/Plugins/Experimental/LearningCore/Source/LearningTraining/Private/LearningExternalTrainer.cpp b/Engine/Plugins/Experimental/LearningCore/Source/LearningTraining/Private/LearningExternalTrainer.cpp
--- a/Engine/Plugins/Experimental/LearningCore/Source/LearningTraining/Private/LearningExternalTrainer.cpp
+++ b/Engine/Plugins/Experimental/LearningCore/Source/LearningTraining/Private/LearningExternalTrainer.cpp
@@ -82,7 +82,7 @@ namespace UE::Learning
 			*TrainerFileName,
 			LogSettings == ELogSetting::Normal ? TEXT("-l") : TEXT(""),
 			*FileManager.ConvertToAbsolutePathForExternalAppForRead(*(InIntermediatePath / TaskName + FString::FromInt(TaskId))),
-			*Controls.Guid.ToString(PLATFORM_MAC ? EGuidFormats::Short : EGuidFormats::DigitsWithHyphensInBraces));
+			*Controls.GetPlatformGuidString());
 
 		bool bLaunched = TrainingProcess.Launch(
 			FileManager.ConvertToAbsolutePathForExternalAppForRead(*PythonExecutablePath),
@@ -285,7 +285,7 @@ namespace UE::Learning
 		{
 			TSharedPtr<FJsonObject> JsonObject = MakeShared<FJsonObject>();
 			JsonObject->SetNumberField(TEXT("NetworkId"), Index);
-			JsonObject->SetStringField(TEXT("Guid"), *NeuralNetworkSharedMemoryArrayViews[Index].Guid.ToString(PLATFORM_MAC ? EGuidFormats::Short : EGuidFormats::DigitsWithHyphensInBraces));
+			JsonObject->SetStringField(TEXT("Guid"), *NeuralNetworkSharedMemoryArrayViews[Index].GetPlatformGuidString());
 
 			TSharedRef<FJsonValueObject> JsonValue = MakeShared<FJsonValueObject>(JsonObject);
 			NetworkGuidsArray.Add(JsonValue);
@@ -296,14 +296,14 @@ namespace UE::Learning
 		for (const FSharedMemoryExperienceContainer& SharedMemoryExperienceContainer : SharedMemoryExperienceContainers)
 		{
 			TSharedPtr<FJsonObject> JsonObject = MakeShared<FJsonObject>();
-			JsonObject->SetStringField(TEXT("EpisodeStartsGuid"), *SharedMemoryExperienceContainer.EpisodeStarts.Guid.ToString(PLATFORM_MAC ? EGuidFormats::Short : EGuidFormats::DigitsWithHyphensInBraces));
-			JsonObject->SetStringField(TEXT("EpisodeLengthsGuid"), *SharedMemoryExperienceContainer.EpisodeLengths.Guid.ToString(PLATFORM_MAC ? EGuidFormats::Short : EGuidFormats::DigitsWithHyphensInBraces));
-			JsonObject->SetStringField(TEXT("EpisodeCompletionModesGuid"), *SharedMemoryExperienceContainer.EpisodeCompletionModes.Guid.ToString(PLATFORM_MAC ? EGuidFormats::Short : EGuidFormats::DigitsWithHyphensInBraces));
+			JsonObject->SetStringField(TEXT("EpisodeStartsGuid"), *SharedMemoryExperienceContainer.EpisodeStarts.GetPlatformGuidString());
+			JsonObject->SetStringField(TEXT("EpisodeLengthsGuid"), *SharedMemoryExperienceContainer.EpisodeLengths.GetPlatformGuidString());
+			JsonObject->SetStringField(TEXT("EpisodeCompletionModesGuid"), *SharedMemoryExperienceContainer.EpisodeCompletionModes.GetPlatformGuidString());
 
 			TArray<TSharedPtr<FJsonValue>> EpisodeFinalObservationsGuidsArray;
 			for (const TSharedMemoryArrayView<2, float>& EpisodeFinalObservations : SharedMemoryExperienceContainer.EpisodeFinalObservations)
 			{
-				EpisodeFinalObservationsGuidsArray.Add(MakeShared<FJsonValueString>(*EpisodeFinalObservations.Guid.ToString(PLATFORM_MAC ? EGuidFormats::Short : EGuidFormats::DigitsWithHyphensInBraces)));
+				EpisodeFinalObservationsGuidsArray.Add(MakeShared<FJsonValueString>(*EpisodeFinalObservations.GetPlatformGuidString()));
 				
 			}
 			JsonObject->SetArrayField(TEXT("EpisodeFinalObservationsGuids"), EpisodeFinalObservationsGuidsArray);
@@ -311,7 +311,7 @@ namespace UE::Learning
 			TArray<TSharedPtr<FJsonValue>> EpisodeFinalMemoryStatesGuidsArray;
 			for (const TSharedMemoryArrayView<2, float>& EpisodeFinalMemoryStates : SharedMemoryExperienceContainer.EpisodeFinalMemoryStates)
 			{
-				EpisodeFinalMemoryStatesGuidsArray.Add(MakeShared<FJsonValueString>(*EpisodeFinalMemoryStates.Guid.ToString(PLATFORM_MAC ? EGuidFormats::Short : EGuidFormats::DigitsWithHyphensInBraces)));
+				EpisodeFinalMemoryStatesGuidsArray.Add(MakeShared<FJsonValueString>(*EpisodeFinalMemoryStates.GetPlatformGuidString()));
 
 			}
 			JsonObject->SetArrayField(TEXT("EpisodeFinalMemoryStatesGuids"), EpisodeFinalMemoryStatesGuidsArray);
@@ -319,7 +319,7 @@ namespace UE::Learning
 			TArray<TSharedPtr<FJsonValue>> ObservationsGuidsArray;
 			for (const TSharedMemoryArrayView<2, float>& Observations : SharedMemoryExperienceContainer.Observations)
 			{
-				ObservationsGuidsArray.Add(MakeShared<FJsonValueString>(*Observations.Guid.ToString(PLATFORM_MAC ? EGuidFormats::Short : EGuidFormats::DigitsWithHyphensInBraces)));
+				ObservationsGuidsArray.Add(MakeShared<FJsonValueString>(*Observations.GetPlatformGuidString()));
 
 			}
 			JsonObject->SetArrayField(TEXT("ObservationsGuids"), ObservationsGuidsArray);
@@ -327,7 +327,7 @@ namespace UE::Learning
 			TArray<TSharedPtr<FJsonValue>> ActionsGuidsArray;
 			for (const TSharedMemoryArrayView<2, float>& Actions : SharedMemoryExperienceContainer.Actions)
 			{
-				ActionsGuidsArray.Add(MakeShared<FJsonValueString>(*Actions.Guid.ToString(PLATFORM_MAC ? EGuidFormats::Short : EGuidFormats::DigitsWithHyphensInBraces)));
+				ActionsGuidsArray.Add(MakeShared<FJsonValueString>(*Actions.GetPlatformGuidString()));
 
 			}
 			JsonObject->SetArrayField(TEXT("ActionsGuids"), ActionsGuidsArray);
@@ -335,7 +335,7 @@ namespace UE::Learning
 			TArray<TSharedPtr<FJsonValue>> ActionModifiersGuidsArray;
 			for (const TSharedMemoryArrayView<2, float>& ActionModifiers : SharedMemoryExperienceContainer.ActionModifiers)
 			{
-				ActionModifiersGuidsArray.Add(MakeShared<FJsonValueString>(*ActionModifiers.Guid.ToString(PLATFORM_MAC ? EGuidFormats::Short : EGuidFormats::DigitsWithHyphensInBraces)));
+				ActionModifiersGuidsArray.Add(MakeShared<FJsonValueString>(*ActionModifiers.GetPlatformGuidString()));
 
 			}
 			JsonObject->SetArrayField(TEXT("ActionModifiersGuids"), ActionModifiersGuidsArray);
@@ -344,7 +344,7 @@ namespace UE::Learning
 			TArray<TSharedPtr<FJsonValue>> MemoryStatesGuidsArray;
 			for (const TSharedMemoryArrayView<2, float>& MemoryStates : SharedMemoryExperienceContainer.MemoryStates)
 			{
-				MemoryStatesGuidsArray.Add(MakeShared<FJsonValueString>(*MemoryStates.Guid.ToString(PLATFORM_MAC ? EGuidFormats::Short : EGuidFormats::DigitsWithHyphensInBraces)));
+				MemoryStatesGuidsArray.Add(MakeShared<FJsonValueString>(*MemoryStates.GetPlatformGuidString()));
 
 			}
 			JsonObject->SetArrayField(TEXT("MemoryStatesGuids"), MemoryStatesGuidsArray);
@@ -352,7 +352,7 @@ namespace UE::Learning
 			TArray<TSharedPtr<FJsonValue>> RewardsGuidsArray;
 			for (const TSharedMemoryArrayView<2, float>& Rewards : SharedMemoryExperienceContainer.Rewards)
 			{
-				RewardsGuidsArray.Add(MakeShared<FJsonValueString>(*Rewards.Guid.ToString(PLATFORM_MAC ? EGuidFormats::Short : EGuidFormats::DigitsWithHyphensInBraces)));
+				RewardsGuidsArray.Add(MakeShared<FJsonValueString>(*Rewards.GetPlatformGuidString()));
 
 			}
 			JsonObject->SetArrayField(TEXT("RewardsGuids"), RewardsGuidsArray);
@@ -366,7 +366,7 @@ namespace UE::Learning
 		FString SharedMemoryConfigString;
 		TSharedRef<TJsonWriter<>> SharedMemoryJsonWriter = TJsonWriterFactory<>::Create(&SharedMemoryConfigString, 0);
 		FJsonSerializer::Serialize(SharedMemoryConfigObject, SharedMemoryJsonWriter, true);
-		FFileHelper::SaveStringToFile(SharedMemoryConfigString, *(ConfigPath / FString::Printf(TEXT("shared-memory-%s.json"), *Controls.Guid.ToString(PLATFORM_MAC ? EGuidFormats::Short : EGuidFormats::DigitsWithHyphensInBraces))));
+		FFileHelper::SaveStringToFile(SharedMemoryConfigString, *(ConfigPath / FString::Printf(TEXT("shared-memory-%s.json"), *Controls.GetPlatformGuidString())));
 
 		// Merge model config into data config for shared memory protocol compatibility
 		for (const auto& Pair : ModelConfigObject->Values)
diff --git a/Engine/Plugins/Experimental/LearningCore/Source/LearningTraining/Public/LearningSharedMemory.h b/Engine/Plugins/Experimental/LearningCore/Source/LearningTraining/Public/LearningSharedMemory.h
--- a/Engine/Plugins/Experimental/LearningCore/Source/LearningTraining/Public/LearningSharedMemory.h
+++ b/Engine/Plugins/Experimental/LearningCore/Source/LearningTraining/Public/LearningSharedMemory.h
@@ -25,6 +25,16 @@ namespace UE::Learning
 		FGuid Guid;
 		TLearningArrayView<DimNum, ElementType> View;
 		FPlatformMemory::FSharedMemoryRegion* Region = nullptr;
+
+		/** Gets a platform specific Guid string. This is required as guid strings on Mac must shorter to be used with shared memory */
+		FString GetPlatformGuidString() const
+		{
+#if PLATFORM_MAC
+			return Guid.ToString(EGuidFormats::Short)
+#else
+			return Guid.ToString(EGuidFormats::DigitsWithHyphensInBraces);
+#endif
+		}
 	};
 
 	namespace SharedMemory
