LearningCore: Fix compilation issue on Mac

#Synced-CL 52297216

diff --git a/Engine/Plugins/Experimental/LearningCore/Source/LearningTraining/Public/LearningSharedMemory.h b/Engine/Plugins/Experimental/LearningCore/Source/LearningTraining/Public/LearningSharedMemory.h
--- a/Engine/Plugins/Experimental/LearningCore/Source/LearningTraining/Public/LearningSharedMemory.h
+++ b/Engine/Plugins/Experimental/LearningCore/Source/LearningTraining/Public/LearningSharedMemory.h
@@ -30,7 +30,7 @@ namespace UE::Learning
 		FString GetPlatformGuidString() const
 		{
 #if PLATFORM_MAC
-			return Guid.ToString(EGuidFormats::Short)
+			return Guid.ToString(EGuidFormats::Short);
 #else
 			return Guid.ToString(EGuidFormats::DigitsWithHyphensInBraces);
 #endif
