LearningCore: Added Update number to Neural Network data asset

#Synced-CL 52510698

diff --git a/Engine/Plugins/Experimental/LearningCore/Source/Learning/Private/LearningNeuralNetwork.cpp b/Engine/Plugins/Experimental/LearningCore/Source/Learning/Private/LearningNeuralNetwork.cpp
--- a/Engine/Plugins/Experimental/LearningCore/Source/Learning/Private/LearningNeuralNetwork.cpp
+++ b/Engine/Plugins/Experimental/LearningCore/Source/Learning/Private/LearningNeuralNetwork.cpp
@@ -26,6 +26,7 @@ void ULearningNeuralNetworkData::Init(const int32 InInputSize, const int32 InOut
 	OutputSize = InOutputSize;
 	CompatibilityHash = InCompatibilityHash;
 	FileData = InFileData;
+	UpdateNumber = 0;
 	UpdateNetwork();
 }
 
@@ -35,6 +36,7 @@ void ULearningNeuralNetworkData::InitFrom(const ULearningNeuralNetworkData* Othe
 	OutputSize = OtherNetworkData->OutputSize;
 	CompatibilityHash = OtherNetworkData->CompatibilityHash;
 	FileData = OtherNetworkData->FileData;
+	UpdateNumber = 0;
 	UpdateNetwork();
 }
 
@@ -74,6 +76,7 @@ bool ULearningNeuralNetworkData::LoadFromSnapshot(const TArrayView<const uint8>
 	int32 FileDataSize = 0;
 	UE::Learning::DeserializeFromBytes(Offset, InBytes, FileDataSize);
 	FileData = InBytes.Slice(Offset, FileDataSize);
+	UpdateNumber++;
 
 	UpdateNetwork();
 	return true;
@@ -133,6 +136,11 @@ int32 ULearningNeuralNetworkData::GetContentHash() const
 	return ContentHash;
 }
 
+int32 ULearningNeuralNetworkData::GetUpdateNumber() const
+{
+	return UpdateNumber;
+}
+
 void ULearningNeuralNetworkData::UpdateNetwork()
 {
 	if (FileData.Num() > 0)
diff --git a/Engine/Plugins/Experimental/LearningCore/Source/Learning/Public/LearningNeuralNetwork.h b/Engine/Plugins/Experimental/LearningCore/Source/Learning/Public/LearningNeuralNetwork.h
--- a/Engine/Plugins/Experimental/LearningCore/Source/Learning/Public/LearningNeuralNetwork.h
+++ b/Engine/Plugins/Experimental/LearningCore/Source/Learning/Public/LearningNeuralNetwork.h
@@ -85,6 +85,9 @@ public:
 	// Gets the content hash
 	UE_API int32 GetContentHash() const;
 
+	// Gets the update number which is incremented each time the network is loaded from a snapshot
+	UE_API int32 GetUpdateNumber() const;
+
 private:
 
 	// Calls UpdateNetwork to update the internal representation.
@@ -113,6 +116,10 @@ private:
 	UPROPERTY()
 	TArray<uint8> FileData;
 
+	// Number which is incremented each time the network is loaded from a snapshot
+	UPROPERTY(VisibleAnywhere, Category = "Network Properties")
+	int32 UpdateNumber = 0;
+
 	// Model Data used by NNE
 	UPROPERTY()
 	TObjectPtr<UNNEModelData> ModelData;
