LearningCore: Added Sparse and NamedSparse observations

#Synced-CL 52698189

diff --git a/Engine/Plugins/Experimental/LearningCore/Source/Learning/Private/LearningObservation.cpp b/Engine/Plugins/Experimental/LearningCore/Source/Learning/Private/LearningObservation.cpp
--- a/Engine/Plugins/Experimental/LearningCore/Source/Learning/Private/LearningObservation.cpp
+++ b/Engine/Plugins/Experimental/LearningCore/Source/Learning/Private/LearningObservation.cpp
@@ -435,6 +435,42 @@ namespace UE::Learning::Observation
 		return { Index, Generation };
 	}
 
+	FSchemaElement FSchema::CreateSparse(const FSchemaSparseParameters Parameters, const FName Tag)
+	{
+		FSparseData ElementData;
+		ElementData.Num = Parameters.Num;
+
+		const int32 Index = Types.Add(EType::Sparse);
+		Tags.Add(Tag);
+		ObservationVectorSizes.Add(Parameters.Num);
+		EncodedVectorSizes.Add(Parameters.Num);
+		ObservationDistributionVectorSizes.Add(Parameters.Num);
+		TypeDataIndices.Add(SparseData.Add(ElementData));
+
+		return { Index, Generation };
+	}
+
+	FSchemaElement FSchema::CreateNamedSparse(const FSchemaNamedSparseParameters Parameters, const FName Tag)
+	{
+		check(!Private::ContainsDuplicates(Parameters.ElementNames));
+
+		FNamedSparseData ElementData;
+		ElementData.Num = Parameters.ElementNames.Num();
+		ElementData.ElementsOffset = SubElementObjects.Num();
+
+		SubElementNames.Append(Parameters.ElementNames);
+		for (int32 Idx = 0; Idx < ElementData.Num; Idx++) { SubElementObjects.Add(FSchemaElement()); }
+
+		const int32 Index = Types.Add(EType::NamedSparse);
+		Tags.Add(Tag);
+		ObservationVectorSizes.Add(Parameters.ElementNames.Num());
+		EncodedVectorSizes.Add(Parameters.ElementNames.Num());
+		ObservationDistributionVectorSizes.Add(Parameters.ElementNames.Num());
+		TypeDataIndices.Add(NamedSparseData.Add(ElementData));
+
+		return { Index, Generation };
+	}
+
 	bool FSchema::IsValid(const FSchemaElement Element) const
 	{
 		return Element.Generation == Generation && Element.Index != INDEX_NONE;
@@ -639,6 +675,26 @@ namespace UE::Learning::Observation
 		return Parameters;
 	}
 
+	FSchemaSparseParameters FSchema::GetSparse(const FSchemaElement Element) const
+	{
+		check(IsValid(Element) && GetType(Element) == EType::Sparse);
+		const FSparseData& ElementData = SparseData[TypeDataIndices[Element.Index]];
+
+		FSchemaSparseParameters Parameters;
+		Parameters.Num = ElementData.Num;
+		return Parameters;
+	}
+
+	FSchemaNamedSparseParameters FSchema::GetNamedSparse(const FSchemaElement Element) const
+	{
+		check(IsValid(Element) && GetType(Element) == EType::NamedSparse);
+		const FNamedSparseData& ElementData = NamedSparseData[TypeDataIndices[Element.Index]];
+
+		FSchemaNamedSparseParameters Parameters;
+		Parameters.ElementNames = TArrayView<const FName>(SubElementNames.GetData() + ElementData.ElementsOffset, ElementData.Num);
+		return Parameters;
+	}
+
 	void FSchema::SetGeneration(const uint32 InGeneration)
 	{
 		Generation = InGeneration;
@@ -671,6 +727,8 @@ namespace UE::Learning::Observation
 		EncodingData.Empty();
 		Conv1dData.Empty();
 		Conv2dData.Empty();
+		SparseData.Empty();
+		NamedSparseData.Empty();
 
 		SubElementNames.Empty();
 		SubElementObjects.Empty();
@@ -708,6 +766,8 @@ namespace UE::Learning::Observation
 		EncodingData.Reset();
 		Conv1dData.Reset();
 		Conv2dData.Reset();
+		SparseData.Reset();
+		NamedSparseData.Reset();
 
 		SubElementNames.Reset();
 		SubElementObjects.Reset();
@@ -721,21 +781,19 @@ namespace UE::Learning::Observation
 	bool FSchema::Serialize(FArchive& Ar)
 	{
 		static constexpr int32 Magic = 0x7db41964;
-		static constexpr int32 Version = 1;
+		static constexpr int32 Version = 2;
 
-		bool bSerializeDeprecatedContinuousData = false;
+		int32 LoadMagic = 0;
+		int32 LoadVersion = Version;
 
 		if (Ar.IsLoading())
 		{
 			int64 Offset = Ar.Tell();
-			int32 LoadMagic = 0;
-			int32 LoadVersion = 0;
 			Ar << LoadMagic;
 			if (!ensure(LoadMagic == Magic)) { Ar.Seek(Offset); Empty(); return false; }
 			Ar << LoadVersion;
 
-			if (LoadVersion == 0) { bSerializeDeprecatedContinuousData = true; }
-			else if (!ensure(LoadVersion == Version)) { Ar.Seek(Offset); Empty(); return false; }
+			if (LoadVersion > Version) { Ar.Seek(Offset); Empty(); return false; }
 		}
 		else
 		{
@@ -754,26 +812,34 @@ namespace UE::Learning::Observation
 		Ar << ObservationDistributionVectorSizes;
 		Ar << TypeDataIndices;
 
-		if (bSerializeDeprecatedContinuousData)
+		if (LoadVersion == 0)
 		{
 			TArray<Private::FContinuousDataDeprecated> DeprecatedContinuousData;
 			Private::ArraySerializeRaw(Ar, DeprecatedContinuousData);
 
+			int32 TotalSize = 0;
 			for (const Private::FContinuousDataDeprecated& DeprecatedData : DeprecatedContinuousData)
 			{
-				FContinuousData CurrentData;
+				TotalSize += DeprecatedData.Num;
+			}
+
+			ContinuousData.Reserve(DeprecatedContinuousData.Num());
+			NormalizationScales.Reserve(TotalSize);
+			NormalizationOffsets.Reserve(TotalSize);
 
+			for (const Private::FContinuousDataDeprecated& DeprecatedData : DeprecatedContinuousData)
+			{
+				FContinuousData CurrentData;
 				CurrentData.Num = DeprecatedData.Num;
 				CurrentData.NormalizationOption = ENormalization::Manual;
 				CurrentData.NormalizationIndexOffset = NormalizationScales.Num();
+				ContinuousData.Add(CurrentData);
 
-				for (int32 i = 0; i < DeprecatedData.Num; ++i)
+				for (int32 Idx = 0; Idx < DeprecatedData.Num; Idx++)
 				{
 					NormalizationScales.Add(DeprecatedData.Scale);
 					NormalizationOffsets.Add(0.0f);
 				}
-
-				ContinuousData.Add(CurrentData);
 			}
 		}
 		else
@@ -794,10 +860,16 @@ namespace UE::Learning::Observation
 		Private::ArraySerializeRaw(Ar, Conv1dData);
 		Private::ArraySerializeRaw(Ar, Conv2dData);
 
+		if (LoadVersion > 1)
+		{
+			Private::ArraySerializeRaw(Ar, SparseData);
+			Private::ArraySerializeRaw(Ar, NamedSparseData);
+		}
+
 		Ar << SubElementNames;
 		Ar << SubElementObjects;
 
-		if (!bSerializeDeprecatedContinuousData)
+		if (LoadVersion > 0)
 		{
 			Private::ArraySerializeRaw(Ar, NormalizationScales);
 			Private::ArraySerializeRaw(Ar, NormalizationOffsets);
@@ -1108,6 +1180,54 @@ namespace UE::Learning::Observation
 		return { Index, Generation };
 	}
 
+
+	FObjectElement FObject::CreateSparse(const FObjectSparseParameters Parameters, const FName Tag)
+	{
+		check(Parameters.DiscreteIndices.Num() == Parameters.Values.Num());
+
+		const int32 Index = Types.Add(EType::Sparse);
+		Tags.Add(Tag);
+
+		ContinuousDataOffsets.Add(ContinuousValues.Num());
+		ContinuousDataNums.Add(Parameters.Values.Num());
+
+		DiscreteDataOffsets.Add(DiscreteValues.Num());
+		DiscreteDataNums.Add(Parameters.DiscreteIndices.Num());
+
+		SubElementDataOffsets.Add(SubElementObjects.Num());
+		SubElementDataNums.Add(0);
+
+		DiscreteValues.Append(Parameters.DiscreteIndices);
+		ContinuousValues.Append(Parameters.Values);
+
+		return { Index, Generation };
+	}
+
+	FObjectElement FObject::CreateNamedSparse(const FObjectNamedSparseParameters Parameters, const FName Tag)
+	{
+		check(Parameters.ElementNames.Num() == Parameters.Values.Num());
+		check(!Private::ContainsDuplicates(Parameters.ElementNames));
+
+		const int32 Index = Types.Add(EType::NamedSparse);
+		Tags.Add(Tag);
+
+		ContinuousDataOffsets.Add(ContinuousValues.Num());
+		ContinuousDataNums.Add(Parameters.Values.Num());
+
+		DiscreteDataOffsets.Add(DiscreteValues.Num());
+		DiscreteDataNums.Add(0);
+
+		SubElementDataOffsets.Add(SubElementObjects.Num());
+		SubElementDataNums.Add(Parameters.ElementNames.Num());
+
+		SubElementNames.Append(Parameters.ElementNames);
+		for (int32 Idx = 0; Idx < Parameters.ElementNames.Num(); Idx++) { SubElementObjects.Add(FObjectElement()); }
+
+		ContinuousValues.Append(Parameters.Values);
+
+		return { Index, Generation };
+	}
+
 	bool FObject::IsValid(const FObjectElement Element) const
 	{
 		return Element.Generation == Generation && Element.Index != INDEX_NONE;
@@ -1245,6 +1365,26 @@ namespace UE::Learning::Observation
 		return Parameters;
 	}
 
+	FObjectSparseParameters FObject::GetSparse(const FObjectElement Element) const
+	{
+		check(IsValid(Element) && GetType(Element) == EType::Sparse);
+
+		FObjectSparseParameters Parameters;
+		Parameters.DiscreteIndices = TArrayView<const int32>(DiscreteValues.GetData() + DiscreteDataOffsets[Element.Index], DiscreteDataNums[Element.Index]);
+		Parameters.Values = TArrayView<const float>(ContinuousValues.GetData() + ContinuousDataOffsets[Element.Index], ContinuousDataNums[Element.Index]);
+		return Parameters;
+	}
+
+	FObjectNamedSparseParameters FObject::GetNamedSparse(const FObjectElement Element) const
+	{
+		check(IsValid(Element) && GetType(Element) == EType::NamedSparse);
+
+		FObjectNamedSparseParameters Parameters;
+		Parameters.ElementNames = TArrayView<const FName>(SubElementNames.GetData() + SubElementDataOffsets[Element.Index], SubElementDataNums[Element.Index]);
+		Parameters.Values = TArrayView<const float>(ContinuousValues.GetData() + ContinuousDataOffsets[Element.Index], ContinuousDataNums[Element.Index]);
+		return Parameters;
+	}
+
 	void FObject::SetGeneration(const uint32 InGeneration)
 	{
 		Generation = InGeneration;
@@ -1525,6 +1665,17 @@ namespace UE::Learning::Observation
 				GetSchemaObjectsCompatibilityHash(Schema, Parameters.Element, Salt) });
 		}
 
+		case EType::Sparse:
+		{
+			return Private::HashCombine({ Hash, Private::HashInt(Schema.GetSparse(SchemaElement).Num) });
+		}
+
+		case EType::NamedSparse:
+		{
+			const FSchemaNamedSparseParameters Parameters = Schema.GetNamedSparse(SchemaElement);
+			return Private::HashCombine({ Hash, Private::HashElements(Schema, Parameters.ElementNames, Salt) });
+		}
+
 		default:
 		{
 			checkNoEntry();
@@ -1699,6 +1850,24 @@ namespace UE::Learning::Observation
 				&& AreSchemaObjectsCompatible(SchemaA, ParametersA.Element, SchemaB, ParametersB.Element);
 		}
 
+		case EType::Sparse: return SchemaA.GetSparse(SchemaElementA).Num == SchemaB.GetSparse(SchemaElementB).Num;
+
+		case EType::NamedSparse:
+		{
+			const FSchemaNamedSparseParameters ParametersA = SchemaA.GetNamedSparse(SchemaElementA);
+			const FSchemaNamedSparseParameters ParametersB = SchemaB.GetNamedSparse(SchemaElementB);
+
+			if (ParametersA.ElementNames.Num() != ParametersB.ElementNames.Num()) { return false; }
+
+			for (int32 SchemaElementAIdx = 0; SchemaElementAIdx < ParametersA.ElementNames.Num(); SchemaElementAIdx++)
+			{
+				const int32 SchemaElementBIdx = ParametersB.ElementNames.Find(ParametersA.ElementNames[SchemaElementAIdx]);
+				if (SchemaElementBIdx == INDEX_NONE) { return false; }
+			}
+
+			return true;
+		}
+
 		default:
 		{
 			checkNoEntry();
@@ -2040,6 +2209,28 @@ namespace UE::Learning::Observation
 			break;
 		}
 
+		case EType::Sparse:
+		{
+			const int32 ValueNum = Schema.GetSparse(SchemaElement).Num;
+
+			OutElement = Builder.MakeDenormalize(
+				ValueNum,
+				Builder.MakeValuesZero(ValueNum),
+				Builder.MakeValuesOne(ValueNum));
+			break;
+		}
+
+		case EType::NamedSparse:
+		{
+			const int32 ValueNum = Schema.GetNamedSparse(SchemaElement).ElementNames.Num();
+
+			OutElement = Builder.MakeDenormalize(
+				ValueNum,
+				Builder.MakeValuesZero(ValueNum),
+				Builder.MakeValuesOne(ValueNum));
+			break;
+		}
+
 		default:
 		{
 			checkNoEntry();
@@ -2405,6 +2596,46 @@ namespace UE::Learning::Observation
 			return;
 		}
 
+
+		case EType::Sparse:
+		{
+			const TArrayView<const int32> ObservationValues = Object.GetSparse(ObjectElement).DiscreteIndices;
+			const TArrayView<const float> ObservationWeights = Object.GetSparse(ObjectElement).Values;
+			check(ObservationValues.Num() == ObservationWeights.Num());
+			check(Schema.GetObservationVectorSize(SchemaElement) >= ObservationValues.Num());
+			check(Schema.GetObservationVectorSize(SchemaElement) >= ObservationWeights.Num());
+			check(Schema.GetObservationVectorSize(SchemaElement) == OutObservationVector.Num());
+
+			// Set values in the observation vector
+
+			for (int32 ObservationValueIdx = 0; ObservationValueIdx < ObservationValues.Num(); ObservationValueIdx++)
+			{
+				check(Schema.GetObservationVectorSize(SchemaElement) > ObservationValues[ObservationValueIdx] && ObservationValues[ObservationValueIdx] >= 0);
+				OutObservationVector[ObservationValues[ObservationValueIdx]] = ObservationWeights[ObservationValueIdx];
+			}
+
+			return;
+		}
+
+		case EType::NamedSparse:
+		{
+			const TArrayView<const FName> SchemaNames = Schema.GetNamedSparse(SchemaElement).ElementNames;
+			const TArrayView<const FName> ObservationValues = Object.GetNamedSparse(ObjectElement).ElementNames;
+			const TArrayView<const float> ObservationWeights = Object.GetNamedSparse(ObjectElement).Values;
+			check(ObservationValues.Num() == ObservationWeights.Num());
+			check(Schema.GetObservationVectorSize(SchemaElement) >= ObservationValues.Num());
+			check(Schema.GetObservationVectorSize(SchemaElement) == OutObservationVector.Num());
+
+			// Set values in the observation vector
+			for (int32 ObservationValueIdx = 0; ObservationValueIdx < ObservationValues.Num(); ObservationValueIdx++)
+			{
+				const int32 ObservationIndex = SchemaNames.Find(ObservationValues[ObservationValueIdx]);
+				check(ObservationIndex != INDEX_NONE);
+				OutObservationVector[ObservationIndex] = ObservationWeights[ObservationValueIdx];
+			}
+			return;
+		}
+
 		default:
 		{
 			checkNoEntry();
@@ -2778,6 +3009,52 @@ namespace UE::Learning::Observation
 			return;
 		}
 
+
+		case EType::Sparse:
+		{
+			check(ObservationVectorSize == Schema.GetSparse(SchemaElement).Num);
+
+			// Find Indices
+			TArray<int32, TInlineAllocator<8>> InclusiveIndices;
+			TArray<float, TInlineAllocator<8>> InclusiveValues;
+			InclusiveIndices.Reserve(ObservationVectorSize);
+			InclusiveValues.Reserve(ObservationVectorSize);
+			for (int32 Idx = 0; Idx < ObservationVectorSize; Idx++)
+			{
+				if (ObservationVector[Idx] != 0.0f)
+				{
+					InclusiveIndices.Add(Idx);
+					InclusiveValues.Add(ObservationVector[Idx]);
+				}
+			}
+
+			OutObjectElement = OutObject.CreateSparse({ InclusiveIndices, InclusiveValues }, SchemaElementTag);
+			return;
+		}
+
+		case EType::NamedSparse:
+		{
+			const TArrayView<const FName> SchemaNames = Schema.GetNamedSparse(SchemaElement).ElementNames;
+			check(ObservationVectorSize == Schema.GetNamedSparse(SchemaElement).ElementNames.Num());
+
+			// Find Names
+			TArray<FName, TInlineAllocator<8>> InclusiveNames;
+			TArray<float, TInlineAllocator<8>> InclusiveValues;
+			InclusiveNames.Reserve(ObservationVectorSize);
+			InclusiveValues.Reserve(ObservationVectorSize);
+			for (int32 Idx = 0; Idx < ObservationVectorSize; Idx++)
+			{
+				if (ObservationVector[Idx] != 0.0f)
+				{
+					InclusiveNames.Add(SchemaNames[Idx]);
+					InclusiveValues.Add(ObservationVector[Idx]);
+				}
+			}
+
+			OutObjectElement = OutObject.CreateNamedSparse({ InclusiveNames, InclusiveValues }, SchemaElementTag);
+			return;
+		}
+
 		default:
 		{
 			checkNoEntry();
@@ -2787,7 +3064,7 @@ namespace UE::Learning::Observation
 		}
 	}
 
-	void GetContinuousActionVectorMask(
+	void GetContinuousObservationVectorMask(
 		TLearningArrayView<1, bool> OutObservationVectorMask,
 		const TLearningArrayView<1, const float> InObservationVector,
 		const FSchema& Schema,
@@ -2846,7 +3123,7 @@ namespace UE::Learning::Observation
 			{
 				const int32 SubElementSize = Schema.GetObservationVectorSize(Parameters.Elements[SchemaElementIdx]);
 
-				GetContinuousActionVectorMask(
+				GetContinuousObservationVectorMask(
 					OutObservationVectorMask.Slice(SubElementOffset, SubElementSize),
 					InObservationVector.Slice(SubElementOffset, SubElementSize),
 					Schema,
@@ -2878,7 +3155,7 @@ namespace UE::Learning::Observation
 
 			const int32 SubElementSize = Schema.GetObservationVectorSize(Parameters.Elements[SchemaElementIndex]);
 
-			GetContinuousActionVectorMask(
+			GetContinuousObservationVectorMask(
 				OutObservationVectorMask.Slice(0, SubElementSize),
 				InObservationVector.Slice(0, SubElementSize),
 				Schema,
@@ -2903,7 +3180,7 @@ namespace UE::Learning::Observation
 
 				if (InObservationVector[TotalSubElementSize + SubElementIdx] == 1.0f)
 				{
-					GetContinuousActionVectorMask(
+					GetContinuousObservationVectorMask(
 						OutObservationVectorMask.Slice(SubElementOffset, SubElementSize),
 						InObservationVector.Slice(SubElementOffset, SubElementSize),
 						Schema,
@@ -2924,7 +3201,7 @@ namespace UE::Learning::Observation
 
 			for (int32 ElementIdx = 0; ElementIdx < Parameters.Num; ElementIdx++)
 			{
-				GetContinuousActionVectorMask(
+				GetContinuousObservationVectorMask(
 					OutObservationVectorMask.Slice(ElementIdx * SubElementSize, SubElementSize),
 					InObservationVector.Slice(ElementIdx * SubElementSize, SubElementSize),
 					Schema,
@@ -2950,7 +3227,7 @@ namespace UE::Learning::Observation
 					break;
 				}
 
-				GetContinuousActionVectorMask(
+				GetContinuousObservationVectorMask(
 					OutObservationVectorMask.Slice(SubElementIdx* SubElementSize, SubElementSize),
 					InObservationVector.Slice(SubElementIdx* SubElementSize, SubElementSize),
 					Schema,
@@ -2962,7 +3239,7 @@ namespace UE::Learning::Observation
 
 		case EType::Encoding:
 		{
-			GetContinuousActionVectorMask(
+			GetContinuousObservationVectorMask(
 				OutObservationVectorMask,
 				InObservationVector,
 				Schema,
@@ -2973,7 +3250,7 @@ namespace UE::Learning::Observation
 
 		case EType::Conv1d:
 		{
-			GetContinuousActionVectorMask(
+			GetContinuousObservationVectorMask(
 				OutObservationVectorMask,
 				InObservationVector,
 				Schema,
@@ -2984,7 +3261,7 @@ namespace UE::Learning::Observation
 
 		case EType::Conv2d:
 		{
-			GetContinuousActionVectorMask(
+			GetContinuousObservationVectorMask(
 				OutObservationVectorMask,
 				InObservationVector,
 				Schema,
@@ -2993,6 +3270,20 @@ namespace UE::Learning::Observation
 			return;
 		}
 
+		case EType::Sparse:
+		{
+			check(ObservationVectorSize == Schema.GetSparse(SchemaElement).Num);
+			Array::Set(OutObservationVectorMask, true);
+			return;
+		}
+
+		case EType::NamedSparse:
+		{
+			check(ObservationVectorSize == Schema.GetNamedSparse(SchemaElement).ElementNames.Num());
+			Array::Set(OutObservationVectorMask, true);
+			return;
+		}
+
 		default:
 		{
 			checkNoEntry();
@@ -3224,6 +3515,30 @@ namespace UE::Learning::Observation
 			return;
 		}
 
+		case EType::Sparse:
+		{
+			check(ObservationVectorSize == Schema.GetSparse(SchemaElement).Num);
+			check(ObservationVectorSize == ObservationDistributionVectorSize);
+			for (int32 Idx = 0; Idx < ObservationVectorSize; Idx++)
+			{
+				InOutObservationDistributionMin[Idx] = FMath::Min(InOutObservationDistributionMin[Idx], InObservationVector[Idx]);
+				InOutObservationDistributionMax[Idx] = FMath::Max(InOutObservationDistributionMax[Idx], InObservationVector[Idx]);
+			}
+			return;
+		}
+
+		case EType::NamedSparse:
+		{
+			check(ObservationVectorSize == Schema.GetNamedSparse(SchemaElement).ElementNames.Num());
+			check(ObservationVectorSize == ObservationDistributionVectorSize);
+			for (int32 Idx = 0; Idx < ObservationVectorSize; Idx++)
+			{
+				InOutObservationDistributionMin[Idx] = FMath::Min(InOutObservationDistributionMin[Idx], InObservationVector[Idx]);
+				InOutObservationDistributionMax[Idx] = FMath::Max(InOutObservationDistributionMax[Idx], InObservationVector[Idx]);
+			}
+			return;
+		}
+
 		default:
 		{
 			checkNoEntry();
@@ -3449,6 +3764,34 @@ namespace UE::Learning::Observation
 			return;
 		}
 
+		case EType::Sparse:
+		{
+			check(ObservationVectorSize == Schema.GetSparse(SchemaElement).Num);
+			check(ObservationVectorSize == ObservationDistributionVectorSize);
+			for (int32 Idx = 0; Idx < ObservationVectorSize; Idx++)
+			{
+				InOutObservationVector[Idx] = FMath::Clamp(
+					InOutObservationVector[Idx],
+					ObservationDistributionMin[Idx],
+					ObservationDistributionMax[Idx]);
+			}
+			return;
+		}
+
+		case EType::NamedSparse:
+		{
+			check(ObservationVectorSize == Schema.GetNamedSparse(SchemaElement).ElementNames.Num());
+			check(ObservationVectorSize == ObservationDistributionVectorSize);
+			for (int32 Idx = 0; Idx < ObservationVectorSize; Idx++)
+			{
+				InOutObservationVector[Idx] = FMath::Clamp(
+					InOutObservationVector[Idx],
+					ObservationDistributionMin[Idx],
+					ObservationDistributionMax[Idx]);
+			}
+			return;
+		}
+
 		default:
 		{
 			checkNoEntry();
diff --git a/Engine/Plugins/Experimental/LearningCore/Source/Learning/Public/LearningObservation.h b/Engine/Plugins/Experimental/LearningCore/Source/Learning/Public/LearningObservation.h
--- a/Engine/Plugins/Experimental/LearningCore/Source/Learning/Public/LearningObservation.h
+++ b/Engine/Plugins/Experimental/LearningCore/Source/Learning/Public/LearningObservation.h
@@ -29,16 +29,16 @@ namespace UE::Learning::Observation
 		// Vector of continuous float observations
 		Continuous = 1,
 
-		// Set of exclusive discrete actions
+		// Set of exclusive discrete observations
 		DiscreteExclusive = 2,
 
-		// Set of inclusive discrete actions
+		// Set of inclusive discrete observations
 		DiscreteInclusive = 3,
 
-		// Set of named exclusive discrete actions
+		// Set of named exclusive discrete observations
 		NamedDiscreteExclusive = 4,
 
-		// Set of named inclusive discrete actions
+		// Set of named inclusive discrete observations
 		NamedDiscreteInclusive = 5,
 
 		// Combination of multiple observations
@@ -64,6 +64,12 @@ namespace UE::Learning::Observation
 
 		// 2d convolution of another observation 
 		Conv2d = 13,
+
+		// Sparse continuous observation indexed by integers
+		Sparse = 14,
+
+		// Sparse continuous observation indexed by names
+		NamedSparse = 15,
 	};
 
 	enum class ENormalization : uint8
@@ -208,13 +214,6 @@ namespace UE::Learning::Observation
 		GELU = 3,
 	};
 
-	/** Enum of Padding Mode */
-	enum class EPaddingMode : uint8
-	{
-		Zeros = 0,
-		Circular = 1,
-	};
-
 	struct FSchemaEncodingParameters
 	{
 		// The sub-element.
@@ -230,6 +229,13 @@ namespace UE::Learning::Observation
 		EEncodingActivationFunction ActivationFunction = EEncodingActivationFunction::ELU;
 	};
 
+	/** Enum of Padding Mode */
+	enum class EPaddingMode : uint8
+	{
+		Zeros = 0,
+		Circular = 1,
+	};
+
 	struct FSchemaConv1dParameters
 	{
 		// The sub-element.
@@ -290,6 +296,18 @@ namespace UE::Learning::Observation
 		EEncodingActivationFunction ActivationFunction = EEncodingActivationFunction::ReLU;
 	};
 
+	struct FSchemaSparseParameters
+	{
+		// Number of values in the sparse observation
+		int32 Num = 0;
+	};
+
+	struct FSchemaNamedSparseParameters
+	{
+		// Names of the sparse observation elements.
+		TArrayView<const FName> ElementNames;
+	};
+
 	/**
 	 * Observation Schema
 	 *
@@ -313,6 +331,8 @@ namespace UE::Learning::Observation
 		UE_API FSchemaElement CreateEncoding(const FSchemaEncodingParameters Parameters, const FName Tag = NAME_None);
 		UE_API FSchemaElement CreateConv1d(const FSchemaConv1dParameters Parameters, const FName Tag = NAME_None);
 		UE_API FSchemaElement CreateConv2d(const FSchemaConv2dParameters Parameters, const FName Tag = NAME_None);
+		UE_API FSchemaElement CreateSparse(const FSchemaSparseParameters Parameters, const FName Tag = NAME_None);
+		UE_API FSchemaElement CreateNamedSparse(const FSchemaNamedSparseParameters Parameters, const FName Tag = NAME_None);
 
 		UE_API FSchemaContinuousParameters GetContinuous(const FSchemaElement Element) const;
 		UE_API FSchemaDiscreteExclusiveParameters GetDiscreteExclusive(const FSchemaElement Element) const;
@@ -327,6 +347,8 @@ namespace UE::Learning::Observation
 		UE_API FSchemaEncodingParameters GetEncoding(const FSchemaElement Element) const;
 		UE_API FSchemaConv1dParameters GetConv1d(const FSchemaElement Element) const;
 		UE_API FSchemaConv2dParameters GetConv2d(const FSchemaElement Element) const;
+		UE_API FSchemaSparseParameters GetSparse(const FSchemaElement Element) const;
+		UE_API FSchemaNamedSparseParameters GetNamedSparse(const FSchemaElement Element) const;
 
 		// Checks if the given element is valid
 		UE_API bool IsValid(const FSchemaElement Element) const;
@@ -466,6 +488,17 @@ namespace UE::Learning::Observation
 			EEncodingActivationFunction ActivationFunction = EEncodingActivationFunction::ReLU;
 		};
 
+		struct FSparseData
+		{
+			int32 Num = INDEX_NONE;
+		};
+
+		struct FNamedSparseData
+		{
+			int32 Num = INDEX_NONE;
+			int32 ElementsOffset = INDEX_NONE;
+		};
+
 		uint32 Generation = 0;
 
 		/** These have entries for each Schema Element */
@@ -490,6 +523,8 @@ namespace UE::Learning::Observation
 		TArray<FEncodingData> EncodingData;
 		TArray<FConv1dData> Conv1dData;
 		TArray<FConv2dData> Conv2dData;
+		TArray<FSparseData> SparseData;
+		TArray<FNamedSparseData> NamedSparseData;
 
 		/** This is an array of all the SubElements and their names, referenced by other elements */
 		TArray<FName> SubElementNames;
@@ -595,6 +630,24 @@ namespace UE::Learning::Observation
 		FObjectElement Element;
 	};
 
+	struct FObjectSparseParameters
+	{
+		// Indices of the sparse observations.
+		TArrayView<const int32> DiscreteIndices;
+
+		// Sparse observation values
+		TArrayView<const float> Values;
+	};
+
+	struct FObjectNamedSparseParameters
+	{
+		// Names of the sparse observations.
+		TArrayView<const FName> ElementNames;
+
+		// Sparse observation values
+		TArrayView<const float> Values;
+	};
+
 	/**
 	 * Observation Object
 	 *
@@ -618,6 +671,8 @@ namespace UE::Learning::Observation
 		UE_API FObjectElement CreateEncoding(const FObjectEncodingParameters Parameters, const FName Tag = NAME_None);
 		UE_API FObjectElement CreateConv1d(const FObjectConv1dParameters Parameters, const FName Tag = NAME_None);
 		UE_API FObjectElement CreateConv2d(const FObjectConv2dParameters Parameters, const FName Tag = NAME_None);
+		UE_API FObjectElement CreateSparse(const FObjectSparseParameters Parameters, const FName Tag = NAME_None);
+		UE_API FObjectElement CreateNamedSparse(const FObjectNamedSparseParameters Parameters, const FName Tag = NAME_None);
 
 		UE_API FObjectContinuousParameters GetContinuous(const FObjectElement Element) const;
 		UE_API FObjectDiscreteExclusiveParameters GetDiscreteExclusive(const FObjectElement Element) const;
@@ -632,6 +687,8 @@ namespace UE::Learning::Observation
 		UE_API FObjectEncodingParameters GetEncoding(const FObjectElement Element) const;
 		UE_API FObjectConv1dParameters GetConv1d(const FObjectElement Element) const;
 		UE_API FObjectConv2dParameters GetConv2d(const FObjectElement Element) const;
+		UE_API FObjectSparseParameters GetSparse(const FObjectElement Element) const;
+		UE_API FObjectNamedSparseParameters GetNamedSparse(const FObjectElement Element) const;
 
 		// Checks if the given element is valid
 		UE_API bool IsValid(const FObjectElement Element) const;
@@ -807,7 +864,7 @@ namespace UE::Learning::Observation
 	 * @param Schema					Observation Schema
 	 * @param SchemaElement				Observation Schema Element
 	 */
-	UE_API void GetContinuousActionVectorMask(
+	UE_API void GetContinuousObservationVectorMask(
 		TLearningArrayView<1, bool> OutObservationVectorMask,
 		const TLearningArrayView<1, const float> InObservationVector,
 		const FSchema& Schema,
diff --git a/Engine/Plugins/Experimental/LearningCore/Source/LearningTraining/Private/LearningTrainer.cpp b/Engine/Plugins/Experimental/LearningCore/Source/LearningTraining/Private/LearningTrainer.cpp
--- a/Engine/Plugins/Experimental/LearningCore/Source/LearningTraining/Private/LearningTrainer.cpp
+++ b/Engine/Plugins/Experimental/LearningCore/Source/LearningTraining/Private/LearningTrainer.cpp
@@ -298,6 +298,31 @@ namespace UE::Learning::Trainer
 			break;
 		}
 
+		case Observation::EType::Sparse:
+		{
+			const Observation::FSchemaSparseParameters Parameters = ObservationSchema.GetSparse(ObservationSchemaElement);
+
+			Object->SetStringField(TEXT("Type"), TEXT("Sparse"));
+			Object->SetNumberField(TEXT("Num"), Parameters.Num);
+			break;
+		}
+
+		case Observation::EType::NamedSparse:
+		{
+			const Observation::FSchemaNamedSparseParameters Parameters = ObservationSchema.GetNamedSparse(ObservationSchemaElement);
+
+			Object->SetStringField(TEXT("Type"), TEXT("NamedSparse"));
+
+			TArray<TSharedPtr<FJsonValue>> ElementNames;
+			for (int32 ElementIdx = 0; ElementIdx < Parameters.ElementNames.Num(); ElementIdx++)
+			{
+				ElementNames.Add(MakeShared<FJsonValueString>(Parameters.ElementNames[ElementIdx].ToString()));
+			}
+
+			Object->SetArrayField(TEXT("ElementNames"), ElementNames);
+			break;
+		}
+
 		default:
 			checkNoEntry();
 		}
@@ -736,6 +761,46 @@ namespace UE::Learning::Trainer
 			return IsObservationSchemaSubsetCompatible(Object->GetObjectField(TEXT("Element")).ToSharedRef(), ObservationSchema, Parameters.Element);
 		}
 
+		case Observation::EType::Sparse:
+		{
+			const Observation::FSchemaSparseParameters Parameters = ObservationSchema.GetSparse(ObservationSchemaElement);
+
+			return
+				Object->GetStringField(TEXT("Type")) == TEXT("Sparse") &&
+				Object->GetNumberField(TEXT("Num")) == Parameters.Num;
+		}
+
+		case Observation::EType::NamedSparse:
+		{
+			const Observation::FSchemaNamedSparseParameters Parameters = ObservationSchema.GetNamedSparse(ObservationSchemaElement);
+
+			if (Object->GetStringField(TEXT("Type")) != TEXT("NamedSparse"))
+			{
+				return false;
+			}
+
+			TArray<TSharedPtr<FJsonValue>> ElementNames = Object->GetArrayField(TEXT("ElementNames"));
+			for (int32 ElementIdx = 0; ElementIdx < Parameters.ElementNames.Num(); ElementIdx++)
+			{
+				bool bFound = false;
+				for (const TSharedPtr<FJsonValue>& ElementName : ElementNames)
+				{
+					if (ElementName->Type == EJson::String && ElementName->AsString() == Parameters.ElementNames[ElementIdx].ToString())
+					{
+						bFound = true;
+						break;
+					}
+				}
+
+				if (!bFound)
+				{
+					return false;
+				}
+			}
+
+			return true;
+		}
+
 		default:
 			checkNoEntry();
 		}
@@ -1010,6 +1075,45 @@ namespace UE::Learning::Trainer
 
 			return;
 		}
+		else if (Object->GetStringField(TEXT("Type")) == TEXT("Sparse"))
+		{
+			const Observation::FSchemaSparseParameters Parameters = ObservationSchema.GetSparse(ObservationSchemaElement);
+
+			for (int32 i = 0; i < Parameters.Num; i++)
+			{
+				Indices.Add(NextIndex);
+				NextIndex++;
+			}
+
+			return;
+		}
+		else if (Object->GetStringField(TEXT("Type")) == TEXT("NamedSparse"))
+		{
+			const Observation::FSchemaNamedSparseParameters Parameters = ObservationSchema.GetNamedSparse(ObservationSchemaElement);
+
+			TArray<TSharedPtr<FJsonValue>> ElementNames = Object->GetArrayField(TEXT("ElementNames"));
+			for (const TSharedPtr<FJsonValue>& ElementName : ElementNames)
+			{
+				bool bFound = false;
+				for (int32 ElementIdx = 0; ElementIdx < Parameters.ElementNames.Num(); ElementIdx++)
+				{
+					if (ElementName->Type == EJson::String && ElementName->AsString() == Parameters.ElementNames[ElementIdx].ToString())
+					{
+						bFound = true;
+						break;
+					}
+				}
+
+				if (bFound)
+				{
+					Indices.Add(NextIndex);
+				}
+
+				NextIndex++;
+			}
+
+			return;
+			}
 		// We didn't return earlier so we received an invalid Type
 		checkNoEntry();
 	}
